Adding user toolbox as experimental plugin

diff --git a/Engine/Plugins/Experimental/UserToolBoxBasicCommand/Config/FilterPlugin.ini b/Engine/Plugins/Experimental/UserToolBoxBasicCommand/Config/FilterPlugin.ini
new file mode 100644
--- /dev/null
+++ b/Engine/Plugins/Experimental/UserToolBoxBasicCommand/Config/FilterPlugin.ini
@@ -0,0 +1,8 @@
+[FilterPlugin]
+; This section lists additional files which will be packaged along with your plugin. Paths should be listed relative to the root plugin directory, and
+; may include "...", "*", and "?" wildcards to match directories, files, and individual characters respectively.
+;
+; Examples:
+;    /README.txt
+;    /Extras/...
+;    /Binaries/ThirdParty/*.dll
diff --git a/Engine/Plugins/Experimental/UserToolBoxBasicCommand/Source/UserToolBoxBasicCommand/Private/ActorFilter/BaseActorFilter.cpp b/Engine/Plugins/Experimental/UserToolBoxBasicCommand/Source/UserToolBoxBasicCommand/Private/ActorFilter/BaseActorFilter.cpp
new file mode 100644
--- /dev/null
+++ b/Engine/Plugins/Experimental/UserToolBoxBasicCommand/Source/UserToolBoxBasicCommand/Private/ActorFilter/BaseActorFilter.cpp
@@ -0,0 +1,269 @@
+// Copyright Epic Games, Inc. All Rights Reserved.
+
+
+#include "ActorFilter/BaseActorFilter.h"
+#include "DatasmithAssetUserData.h"
+#include "DatasmithContentBlueprintLibrary.h"
+#include "Widgets/Layout/SScrollBox.h"
+#include "Widgets/SWindow.h"
+#include "Widgets/Input/SButton.h"
+#include "Editor.h"
+void GetAttachedActors( AActor* Actor, bool bRecurseChildren, TSet< AActor* >& OutActors )
+{
+	TArray< AActor* > ChildrenActors;
+	Actor->GetAttachedActors( ChildrenActors );
+	for ( AActor* ChildActor : ChildrenActors )
+	{
+		OutActors.Add( ChildActor );
+
+		if ( bRecurseChildren )
+		{
+			GetAttachedActors( ChildActor, bRecurseChildren, OutActors );
+		}
+	}
+}
+TArray<AActor*> GetDescendantsAtN(AActor* Actor, int Rank, bool AddIntermediaries)
+{
+	const int CurrentRank=Rank-1;
+	TArray<AActor*> Result;
+	TArray<AActor*> Children;
+	Actor->GetAttachedActors(Children);
+	if (CurrentRank>0)
+	{
+		for (AActor* Child:Children)
+		{
+			Result.Append(GetDescendantsAtN(Child,CurrentRank,AddIntermediaries));
+		}	
+	}
+	else
+	{
+		Result.Add(Actor);
+	}
+	if (AddIntermediaries)
+	{
+		Result.AddUnique(Actor);
+	}
+	return Result;
+	
+	
+	
+	
+}
+
+TArray<AActor*> UBaseActorFilter::Filter_Implementation(const TArray<AActor*>& Source)
+{
+	return Source;
+}
+
+bool UBaseActorFilter::FilterUnit_Implementation(AActor* Source)
+{
+	return true;
+}
+
+ TArray<AActor*> UGetAllDescendants::FilterImpl(const TArray<AActor*>& Source)
+{
+	TSet<AActor*> Result;
+	for (AActor* Actor:Source)
+	{
+		GetAttachedActors(Actor,true,Result);
+	}
+	return Result.Array();
+}
+
+ TArray<AActor*> UGetParents::FilterImpl(const TArray<AActor*>& Source)
+{
+	TSet<AActor*> Result;
+	for (AActor* Actor:Source)
+	{
+		Result.Add(Actor->GetAttachParentActor());
+	}
+	Result.Remove(nullptr);
+	return Result.Array();
+}
+
+ TArray<AActor*> UHasAttachedActor::FilterImpl(const TArray<AActor*>& Source)
+{
+	TSet<AActor*> Result;
+	for (AActor* Actor:Source)
+	{
+		TArray<AActor*> Children;
+		Actor->GetAttachedActors(Children);
+		if (Children.Num()>0)
+		{
+			Result.Add(Actor);
+		}
+			
+	}
+	return Result.Array();
+}
+
+ TArray<AActor*> UHasComponentOfClass::FilterImpl(const TArray<AActor*>& Source)
+{
+	TSet<AActor*> Result;
+	for (AActor* Actor:Source)
+	{
+		TArray<UActorComponent*> Components;
+		Actor->GetComponents(ComponentClass,Components);
+		if (Components.Num()>0)
+		{
+			Result.Add(Actor);
+		}
+			
+	}
+	return Result.Array();
+}
+
+ TArray<AActor*> UHasMetadataByKey::FilterImpl(const TArray<AActor*>& Source)
+{
+	TSet<AActor*> Result;
+	for (AActor* Actor:Source)
+	{
+		FString Value=UDatasmithContentBlueprintLibrary::GetDatasmithUserDataValueForKey(Actor,Key);
+		if (!Value.IsEmpty())
+		{
+			Result.Add(Actor);
+		}
+			
+	}
+	return Result.Array();
+}
+
+ TArray<AActor*> UHasMetadataByKeyAndValue::FilterImpl(const TArray<AActor*>& Source)
+{
+	TSet<AActor*> Result;
+	for (AActor* Actor:Source)
+	{
+		FString RetValue=UDatasmithContentBlueprintLibrary::GetDatasmithUserDataValueForKey(Actor,Key);
+		if (RetValue==Value)
+		{
+			Result.Add(Actor);
+		}
+			
+	}
+	return Result.Array();
+}
+
+ TArray<AActor*> UIsClassOf::FilterImpl(const TArray<AActor*>& Source)
+{
+	TSet<AActor*> Result;
+	for (AActor* Actor:Source)
+	{
+		if ( Actor->GetClass()==ActorClass || (ChildClass && Actor->GetClass()->IsChildOf(ActorClass)))
+		{
+			Result.Add(Actor);
+		}
+			
+	}
+	return Result.Array();
+}
+
+ TArray<AActor*> UGetNDescendants::FilterImpl(const TArray<AActor*>& Source)
+{
+	TSet<AActor*> Result;
+	for (AActor* Actor:Source)
+	{
+		Result.Append(GetDescendantsAtN(Actor,N+1,AddIntermediaries));
+	}
+	return Result.Array();
+}
+
+TArray<AActor*> UHasMetadataByKeyAndValueDropDown::FilterImpl(const TArray<AActor*>& Source)
+{
+	TSet<AActor*> Result;
+	TMap<FString,TArray<AActor*>> ActorSortByValue;
+	TArray<UDatasmithAssetUserData*> UserDatas;
+	TSet<FName> TotalKeys;
+    UDatasmithContentBlueprintLibrary::GetAllDatasmithUserData(AActor::StaticClass(),UserDatas);
+	for (UDatasmithAssetUserData* UserData:UserDatas)
+	{
+		TArray<FName> Keys;
+		UserData->MetaData.GenerateKeyArray(Keys);
+		TotalKeys.Append(Keys);
+	}
+	FName SelectedKey=NAME_None;
+	TSharedRef<SVerticalBox> KeyList = SNew(SVerticalBox);
+	const ISlateStyle& EditorStyle=FAppStyle::Get();
+	TSharedPtr<SWindow> Window = SNew(SWindow)
+		.Title(FText::FromString("Select a Key"))
+		.ClientSize(FVector2D(100.f, FMath::Max(25.f * ActorSortByValue.Num(),300.0f)))
+		.IsTopmostWindow(true);
+
+	for (FName Key : TotalKeys) {
+		KeyList->AddSlot()
+			[
+				SNew(SButton)
+				.Text(FText::FromName(Key))
+				.HAlign(EHorizontalAlignment::HAlign_Center)
+				.VAlign(EVerticalAlignment::VAlign_Center)
+				.ButtonStyle(EditorStyle, "PropertyEditor.AssetComboStyle")
+				.ForegroundColor(EditorStyle.GetColor("PropertyEditor.AssetName.ColorAndOpacity"))
+				.OnClicked_Lambda([Key, Window, &SelectedKey]()->FReply
+				{
+					SelectedKey=Key;
+					Window->RequestDestroyWindow();
+					return FReply::Handled();
+				})
+			];
+	}
+	Window->SetContent(
+		SNew(SScrollBox)
+		+SScrollBox::Slot()
+		[
+		KeyList
+		]
+		);
+	GEditor->EditorAddModalWindow(Window.ToSharedRef());
+	if (SelectedKey==NAME_None)
+	{
+		return TArray<AActor*>();	
+	}
+	
+	
+
+	
+	for (AActor* Actor:Source)
+	{
+		FString Value=UDatasmithContentBlueprintLibrary::GetDatasmithUserDataValueForKey(Actor,SelectedKey);
+		if (!Value.IsEmpty())
+		{
+			ActorSortByValue.FindOrAdd(Value).Add(Actor);
+		}
+	}
+	
+	 Window = SNew(SWindow)
+		.Title(FText::FromString("Select a value"))
+		.ClientSize(FVector2D(100.f, FMath::Max(25.f * ActorSortByValue.Num(),300.0f)))
+		.IsTopmostWindow(true);
+
+	TSharedRef<SVerticalBox> ValueList = SNew(SVerticalBox);
+	TArray<FString> Values;
+	ActorSortByValue.GenerateKeyArray(Values);
+	FString SelectedValue;
+	for (const FString& Value : Values) {
+		ValueList->AddSlot()
+			[
+				SNew(SButton)
+				.Text(FText::FromString(Value))
+				.HAlign(EHorizontalAlignment::HAlign_Center)
+				.VAlign(EVerticalAlignment::VAlign_Center)
+				.ButtonStyle(EditorStyle, "PropertyEditor.AssetComboStyle")
+				.ForegroundColor(EditorStyle.GetColor("PropertyEditor.AssetName.ColorAndOpacity"))
+				.OnClicked_Lambda([Value, Window, &SelectedValue]()->FReply
+				{
+					SelectedValue=Value;
+					Window->RequestDestroyWindow();
+					return FReply::Handled();
+				})
+			];
+	}
+	Window->SetContent(
+		SNew(SScrollBox)
+		+SScrollBox::Slot()
+		[
+			ValueList
+		]
+		);
+	GEditor->EditorAddModalWindow(Window.ToSharedRef());
+	return ActorSortByValue.FindOrAdd(SelectedValue);
+
+}
diff --git a/Engine/Plugins/Experimental/UserToolBoxBasicCommand/Source/UserToolBoxBasicCommand/Private/AssignToLayer.cpp b/Engine/Plugins/Experimental/UserToolBoxBasicCommand/Source/UserToolBoxBasicCommand/Private/AssignToLayer.cpp
new file mode 100644
--- /dev/null
+++ b/Engine/Plugins/Experimental/UserToolBoxBasicCommand/Source/UserToolBoxBasicCommand/Private/AssignToLayer.cpp
@@ -0,0 +1,81 @@
+// Copyright Epic Games, Inc. All Rights Reserved.
+
+
+#include "AssignToLayer.h"
+
+#include "Layers/LayersSubsystem.h"
+#include "Layers/Layer.h"
+#include "Editor.h"
+#include "UserToolBoxBasicCommand.h"
+#include "Widgets/Input/SEditableTextBox.h"
+#include "Widgets/Input/SButton.h"
+
+void UAssignToLayer::Execute()
+{
+	
+	ULayersSubsystem* LayersSubsystem = GEditor->GetEditorSubsystem<ULayersSubsystem>();
+
+	if (LayersSubsystem == nullptr)
+	{
+		UE_LOG(LogUserToolBoxBasicCommand, Warning, TEXT("Could no retrieve LayersSubsystem"));
+		return;
+	}
+	const ISlateStyle& EditorStyle=FAppStyle::Get();
+	TArray<TWeakObjectPtr<ULayer>> Layers;
+	LayersSubsystem->AddAllLayersTo(Layers);
+
+	TSharedRef<SWindow> Window = SNew(SWindow)
+		.Title(NSLOCTEXT("UTBAssignToLayer", "SelectLayer", "Select Layer"))
+		.ClientSize(FVector2D(100.f, 30.f * (Layers.Num()+1)))
+		.IsTopmostWindow(true);
+
+	TSharedRef<SVerticalBox> LayerList = SNew(SVerticalBox)
+	+SVerticalBox::Slot()
+	[
+			SNew(SEditableTextBox)
+			.OnTextCommitted(FOnTextCommitted::CreateLambda([Window, LayersSubsystem](const FText& NewText,ETextCommit::Type CommitType)
+			{
+				LayersSubsystem->AddSelectedActorsToLayer(*NewText.ToString());
+					Window->RequestDestroyWindow();
+			}))
+	];
+
+	
+	for (TWeakObjectPtr<ULayer> Layer : Layers) {
+		LayerList->AddSlot()
+			[
+				SNew(SButton)
+				.Text(FText::FromName(GetLayerName(Layer)))
+				.HAlign(EHorizontalAlignment::HAlign_Center)
+				.VAlign(EVerticalAlignment::VAlign_Center)
+				.ButtonStyle(EditorStyle, "PropertyEditor.AssetComboStyle")
+				.ForegroundColor(EditorStyle.GetColor("PropertyEditor.AssetName.ColorAndOpacity"))
+				.OnClicked_Lambda([Layer, Window, LayersSubsystem, this]()->FReply
+				{
+					LayersSubsystem->AddSelectedActorsToLayer(GetLayerName(Layer));
+
+					Window->RequestDestroyWindow();
+					return FReply::Handled();
+				})
+			];
+	}
+
+	Window->SetContent(LayerList);
+
+	GEditor->EditorAddModalWindow(Window);
+
+	return;
+}
+
+FName UAssignToLayer::GetLayerName(TWeakObjectPtr<ULayer> Layer)
+{
+	return Layer->GetLayerName();
+
+}
+
+UAssignToLayer::UAssignToLayer()
+{
+	Name="Assign to layer";
+	Tooltip="Assign selected actor to a specific layer";
+	Category="Level";
+}
diff --git a/Engine/Plugins/Experimental/UserToolBoxBasicCommand/Source/UserToolBoxBasicCommand/Private/AssignToLevel.cpp b/Engine/Plugins/Experimental/UserToolBoxBasicCommand/Source/UserToolBoxBasicCommand/Private/AssignToLevel.cpp
new file mode 100644
--- /dev/null
+++ b/Engine/Plugins/Experimental/UserToolBoxBasicCommand/Source/UserToolBoxBasicCommand/Private/AssignToLevel.cpp
@@ -0,0 +1,58 @@
+// Copyright Epic Games, Inc. All Rights Reserved.
+
+
+#include "AssignToLevel.h"
+#include "Engine/Selection.h"
+#include "EditorLevelUtils.h"
+#include "Styling/AppStyle.h"
+#include "Editor.h"
+#include "Widgets/Input/SButton.h"
+void UAssignToLevel::Execute()
+{
+	const ISlateStyle& EditorStyle=FAppStyle::Get();
+	
+	TArray<ULevel*> Levels=GEditor->GetEditorWorldContext().World()->GetLevels();
+
+	TSharedRef<SWindow> Window = SNew(SWindow)
+		.Title(NSLOCTEXT("UTBAssignToLevel", "SelectLevel", "Select Level"))
+		.ClientSize(FVector2D(100.f, 25.f * Levels.Num()))
+		.IsTopmostWindow(true);
+
+	TSharedRef<SVerticalBox> LevelList = SNew(SVerticalBox);
+
+	for (ULevel* Level : Levels) {
+		FString FullLevelName = Level->GetFullName();
+		FString LeftS, RightS;
+		FString SearchString(TEXT("."));
+
+		FullLevelName.Split(SearchString, &LeftS, &RightS);
+		SearchString = TEXT(":");
+		RightS.Split(SearchString, &LeftS, &RightS);
+
+		LevelList->AddSlot()
+			[
+				SNew(SButton)
+				.Text(FText::FromString(LeftS))
+				.HAlign(EHorizontalAlignment::HAlign_Center)
+				.VAlign(EVerticalAlignment::VAlign_Center)
+				.ButtonStyle(EditorStyle, "PropertyEditor.AssetComboStyle")
+				.ForegroundColor(EditorStyle.GetColor("PropertyEditor.AssetName.ColorAndOpacity"))
+				.OnClicked_Lambda([Level,Window]()->FReply
+				{
+					EditorLevelUtils::MoveSelectedActorsToLevel(Level);
+					Window->RequestDestroyWindow();
+					return FReply::Handled();
+				})
+			];
+	}
+	Window->SetContent(LevelList);
+	GEditor->EditorAddModalWindow(Window);
+	return;
+}
+
+UAssignToLevel::UAssignToLevel()
+{
+	Name="Assign to level";
+	Tooltip="Assign selected actor into a specific sublevel";
+	Category="Level";
+}
diff --git a/Engine/Plugins/Experimental/UserToolBoxBasicCommand/Source/UserToolBoxBasicCommand/Private/BaseCompositeCommand.cpp b/Engine/Plugins/Experimental/UserToolBoxBasicCommand/Source/UserToolBoxBasicCommand/Private/BaseCompositeCommand.cpp
new file mode 100644
--- /dev/null
+++ b/Engine/Plugins/Experimental/UserToolBoxBasicCommand/Source/UserToolBoxBasicCommand/Private/BaseCompositeCommand.cpp
@@ -0,0 +1,43 @@
+// Copyright Epic Games, Inc. All Rights Reserved.
+
+
+#include "BaseCompositeCommand.h"
+#include "EditorUtilityBlueprint.h"
+#include "UserToolBoxBaseBlueprint.h"
+#include "UTBBaseTab.h"
+#include "Engine/Selection.h"
+#include "UserToolBoxBasicCommand.h"
+
+UUTBBaseCommand* UBaseCompositeCommand::CopyCommand(UObject* Owner) const
+{
+	UUTBTabSection* Section=Cast<UUTBTabSection>(Owner);
+	if (!IsValid(Section))
+	{
+		UE_LOG(LogUserToolBoxBasicCommand, Error, TEXT("The new owner %s isn't a tab section"), *Owner->GetName());
+		return nullptr;
+	}
+	UUserToolBoxBaseTab* NewTab= Cast<UUserToolBoxBaseTab>(Section->GetOuter());
+	if (!IsValid(NewTab))
+	{
+		UE_LOG(LogUserToolBoxBasicCommand, Error, TEXT("The owner of the section isn't a tab "), Section->GetOuter()?(*Section->GetOuter()->GetName()):TEXT("Null"));
+		return nullptr;
+	}
+	UBaseCompositeCommand* NewToggleCommand=DuplicateObject(this,Owner);
+	if (NewTab==this->GetOuter()->GetOuter())
+	{
+		//
+	}
+	else
+	{
+		NewToggleCommand->Commands.Empty(Commands.Num());
+		UUTBTabSection* PlaceHolderSection=NewTab->GetPlaceHolderSection();
+		for (UUTBBaseCommand* CurrentCommand:Commands)
+		{
+			UUTBBaseCommand* NewSubCommand=CurrentCommand->CopyCommand(PlaceHolderSection);
+			NewToggleCommand->Commands.Add(NewSubCommand);
+			PlaceHolderSection->Commands.Add(NewSubCommand);
+		}	
+	}
+	Section->Commands.Add(NewToggleCommand);
+	return NewToggleCommand;
+}
\ No newline at end of file
diff --git a/Engine/Plugins/Experimental/UserToolBoxBasicCommand/Source/UserToolBoxBasicCommand/Private/BaseCompositeInlineCommand.cpp b/Engine/Plugins/Experimental/UserToolBoxBasicCommand/Source/UserToolBoxBasicCommand/Private/BaseCompositeInlineCommand.cpp
new file mode 100644
--- /dev/null
+++ b/Engine/Plugins/Experimental/UserToolBoxBasicCommand/Source/UserToolBoxBasicCommand/Private/BaseCompositeInlineCommand.cpp
@@ -0,0 +1,5 @@
+// Copyright Epic Games, Inc. All Rights Reserved.
+
+
+#include "BaseCompositeInlineCommand.h"
+
diff --git a/Engine/Plugins/Experimental/UserToolBoxBasicCommand/Source/UserToolBoxBasicCommand/Private/ChangeViewMode.cpp b/Engine/Plugins/Experimental/UserToolBoxBasicCommand/Source/UserToolBoxBasicCommand/Private/ChangeViewMode.cpp
new file mode 100644
--- /dev/null
+++ b/Engine/Plugins/Experimental/UserToolBoxBasicCommand/Source/UserToolBoxBasicCommand/Private/ChangeViewMode.cpp
@@ -0,0 +1,26 @@
+// Copyright Epic Games, Inc. All Rights Reserved.
+
+
+#include "ChangeViewMode.h"
+
+#include "LevelEditorViewport.h"
+
+void UChangeViewMode::Execute()
+{
+	const TArray<FLevelEditorViewportClient*>& VCs=GEditor->GetLevelViewportClients();
+	for (FLevelEditorViewportClient* VC:VCs)
+	{
+		VC->SetViewMode(ViewMode);
+	}
+		
+	
+	return Super::Execute();
+}
+
+UChangeViewMode::UChangeViewMode()
+{
+	Name="Change Viewport mode";
+	Category="Viewport";
+	Tooltip="Change the main viewport mode ( lit, unlit ...)";
+	
+}
diff --git a/Engine/Plugins/Experimental/UserToolBoxBasicCommand/Source/UserToolBoxBasicCommand/Private/CleanHierarchy.cpp b/Engine/Plugins/Experimental/UserToolBoxBasicCommand/Source/UserToolBoxBasicCommand/Private/CleanHierarchy.cpp
new file mode 100644
--- /dev/null
+++ b/Engine/Plugins/Experimental/UserToolBoxBasicCommand/Source/UserToolBoxBasicCommand/Private/CleanHierarchy.cpp
@@ -0,0 +1,143 @@
+// Copyright Epic Games, Inc. All Rights Reserved.
+
+
+ #include "CleanHierarchy.h"
+#include "DatasmithContentBlueprintLibrary.h"
+#include "EngineUtils.h"
+#include "Engine/Selection.h"
+#include "Editor.h"
+#include "Engine/Classes/Components/MeshComponent.h"
+
+int GetAttachedChildCount(AActor* Actor)
+{
+	TArray< AActor* > ChildrenActors;
+	Actor->GetAttachedActors( ChildrenActors );
+	return ChildrenActors.Num();
+}
+void GetAttachedActorsRecursive( AActor* Actor, bool bRecurseChildren, TSet< AActor* >& OutActors )
+{
+	TArray< AActor* > ChildrenActors;
+	Actor->GetAttachedActors( ChildrenActors );
+	for ( AActor* ChildActor : ChildrenActors )
+	{
+		OutActors.Add( ChildActor );
+		if ( bRecurseChildren )
+		{
+			GetAttachedActorsRecursive( ChildActor, bRecurseChildren, OutActors );
+		}
+	}
+}
+UCleanHierarchy::UCleanHierarchy()
+{
+	Name="Clean Hierarchy";
+	Tooltip="Remove geometry-less actor from the selected node";
+	Category="Scene";
+}
+
+void UCleanHierarchy::Execute()
+{
+	UWorld* World = GEngine->GetWorldContexts()[0].World();
+
+	USelection* SelectedActors = GEditor->GetSelectedActors();
+	TArray<AActor*> Actors;
+	SelectedActors->GetSelectedObjects(Actors);
+	TSet<AActor*> GeometryLessActors;
+	TSet<AActor*> ChildLessActors;
+	TSet<AActor*> ActorsToProcess;
+	for (AActor* CurrentActor:Actors)
+	{
+		GetAttachedActorsRecursive(CurrentActor,true,ActorsToProcess);
+	}
+	for (AActor* CurrentActor:ActorsToProcess)
+	{
+		
+		TArray<UActorComponent*> Components;
+		CurrentActor->GetComponents(UMeshComponent::StaticClass(),Components,false);
+		if (Components.IsEmpty())
+		{
+			bool FoundMetadata=false;
+			for (FString Key:DSMetadataToPreserve)
+			{
+				if (!UDatasmithContentBlueprintLibrary::GetDatasmithUserDataValueForKey(CurrentActor,FName(Key)).IsEmpty())
+				{
+					FoundMetadata=true;
+					break;
+				}
+			}
+			if (!FoundMetadata)
+			{
+				GeometryLessActors.Add(CurrentActor);
+			}
+		}	
+		if (GetAttachedChildCount(CurrentActor)==0 && GeometryLessActors.Contains(CurrentActor))
+		{
+			ChildLessActors.Add(CurrentActor);
+		}	
+		
+	}
+	TSet<AActor*> ToRemove;
+
+	//first we remove empty branches
+	TArray<AActor*> WorkingArray(ChildLessActors.Array());
+	while (!WorkingArray.IsEmpty())
+	{
+		TArray<AActor*> Parents;
+		for (AActor* CurrentActor:WorkingArray)
+		{
+			if (GetAttachedChildCount(CurrentActor)!=0)
+			{
+				continue;
+			}
+			AActor* Parent=CurrentActor->GetAttachParentActor();
+			if (Parent && GeometryLessActors.Contains(Parent))
+			{
+				Parents.Add(Parent);
+			}
+			CurrentActor->DetachFromActor(FDetachmentTransformRules::KeepWorldTransform);
+			ToRemove.Add(CurrentActor);
+			GeometryLessActors.Remove(CurrentActor);
+		}
+		WorkingArray=Parents;
+	}
+	for (AActor* ElementToRemove:ToRemove)
+	{
+		ElementToRemove->Destroy();
+	}
+	// now we are looking for intermediary
+	for (AActor* CurGeometryLess:GeometryLessActors)
+	{
+		int ChildCount=GetAttachedChildCount(CurGeometryLess);
+		if (ChildCount!=0)
+		{
+			//finding the best parent
+			AActor* PotentialNewParent=CurGeometryLess->GetAttachParentActor();
+			bool FoundParent=false;
+			while (PotentialNewParent->IsValidLowLevel())
+			{
+				if (GeometryLessActors.Contains(PotentialNewParent))
+				{
+					PotentialNewParent=PotentialNewParent->GetAttachParentActor();
+					continue;
+				}
+				break;
+			}
+			TArray<AActor*> Children;
+			CurGeometryLess->GetAttachedActors(Children);
+			for (AActor* Child:Children)
+			{
+				if (PotentialNewParent==nullptr)
+				{
+					Child->DetachFromActor(FDetachmentTransformRules::KeepWorldTransform);
+				}
+				else
+				{
+					Child->AttachToActor(PotentialNewParent,FAttachmentTransformRules::KeepWorldTransform);
+				}
+			}
+		}
+	}
+	for (AActor* ElementToRemove:GeometryLessActors)
+	{
+		ElementToRemove->Destroy();
+	}
+}
diff --git a/Engine/Plugins/Experimental/UserToolBoxBasicCommand/Source/UserToolBoxBasicCommand/Private/CompositeCommand.cpp b/Engine/Plugins/Experimental/UserToolBoxBasicCommand/Source/UserToolBoxBasicCommand/Private/CompositeCommand.cpp
new file mode 100644
--- /dev/null
+++ b/Engine/Plugins/Experimental/UserToolBoxBasicCommand/Source/UserToolBoxBasicCommand/Private/CompositeCommand.cpp
@@ -0,0 +1,28 @@
+// Copyright Epic Games, Inc. All Rights Reserved.
+
+
+#include "CompositeCommand.h"
+
+
+#include "EditorUtilityBlueprint.h"
+#include "EngineUtils.h"
+#include "UserToolBoxBaseBlueprint.h"
+#include "UTBBaseTab.h"
+#include "Engine/Selection.h"
+#include "UserToolBoxBasicCommand.h"
+
+void UCompositeCommand::Execute()
+{
+
+	for (UUTBBaseCommand* Command:Commands)
+	{
+		
+		
+		if (!IsValid(Command))
+		{
+			continue;	
+		}
+		Command->ExecuteCommand();
+	}
+}
+
diff --git a/Engine/Plugins/Experimental/UserToolBoxBasicCommand/Source/UserToolBoxBasicCommand/Private/CompositeInlineCommand.cpp b/Engine/Plugins/Experimental/UserToolBoxBasicCommand/Source/UserToolBoxBasicCommand/Private/CompositeInlineCommand.cpp
new file mode 100644
--- /dev/null
+++ b/Engine/Plugins/Experimental/UserToolBoxBasicCommand/Source/UserToolBoxBasicCommand/Private/CompositeInlineCommand.cpp
@@ -0,0 +1,23 @@
+// Copyright Epic Games, Inc. All Rights Reserved.
+
+
+#include "CompositeInlineCommand.h"
+
+
+
+
+void UCompositeInlineCommand::Execute()
+{
+
+	for (UUTBBaseCommand* Command:Commands)
+	{
+		
+		
+		if (!IsValid(Command))
+		{
+			continue;	
+		}
+		Command->ExecuteCommand();
+	}
+}
+
diff --git a/Engine/Plugins/Experimental/UserToolBoxBasicCommand/Source/UserToolBoxBasicCommand/Private/ConsoleVariable.cpp b/Engine/Plugins/Experimental/UserToolBoxBasicCommand/Source/UserToolBoxBasicCommand/Private/ConsoleVariable.cpp
new file mode 100644
--- /dev/null
+++ b/Engine/Plugins/Experimental/UserToolBoxBasicCommand/Source/UserToolBoxBasicCommand/Private/ConsoleVariable.cpp
@@ -0,0 +1,18 @@
+// Copyright Epic Games, Inc. All Rights Reserved.
+
+
+#include "ConsoleVariable.h"
+#include "HAL/IConsoleManager.h"
+#include "Kismet/KismetSystemLibrary.h"
+#include "UserToolBoxBasicCommand.h"
+#include "Engine/Classes/Engine/World.h"
+void UConsoleVariable::Execute()
+{
+	
+	for (const FString& Command : ConsoleCommands)
+	{
+		UE_LOG(LogUserToolBoxBasicCommand, Log, TEXT("Executing Console Command \"%s\" ."), *Command);
+		UKismetSystemLibrary::ExecuteConsoleCommand(GetWorld(), Command, nullptr);
+	}
+	
+}
diff --git a/Engine/Plugins/Experimental/UserToolBoxBasicCommand/Source/UserToolBoxBasicCommand/Private/Customization/BaseCompositeCommandCustomization.cpp b/Engine/Plugins/Experimental/UserToolBoxBasicCommand/Source/UserToolBoxBasicCommand/Private/Customization/BaseCompositeCommandCustomization.cpp
new file mode 100644
--- /dev/null
+++ b/Engine/Plugins/Experimental/UserToolBoxBasicCommand/Source/UserToolBoxBasicCommand/Private/Customization/BaseCompositeCommandCustomization.cpp
@@ -0,0 +1,130 @@
+﻿// Copyright Epic Games, Inc. All Rights Reserved.
+
+
+#include "BaseCompositeCommandCustomization.h"
+#include "IDetailChildrenBuilder.h"
+#include "DetailCategoryBuilder.h"
+#include "DetailLayoutBuilder.h"
+#include "DetailWidgetRow.h"
+#include "BaseCompositeCommand.h"
+#include "PropertyCustomizationHelpers.h"
+#include "EditorWidgets/Public/SDropTarget.h"
+#include "Editor/UTBTabEditor.h"
+#include "Widgets/Text/STextBlock.h"
+
+void FBaseCompositeCommandCustomization::CustomizeDetails(IDetailLayoutBuilder& DetailBuilder)
+{
+	IDetailCategoryBuilder& CompositeCommandCategory = DetailBuilder.EditCategory("Composite Commands");
+
+	TSharedRef<IPropertyHandle> CommandsProperty = DetailBuilder.GetProperty("Commands");
+
+	TSharedRef<FDetailArrayBuilder> CommandArrayBuilder = MakeShareable( new FDetailArrayBuilder(CommandsProperty ) );
+	CommandArrayBuilder->OnGenerateArrayElementWidget( FOnGenerateArrayElementWidget::CreateSP(this, &FBaseCompositeCommandCustomization::GenerateArrayChildren) );
+
+	
+
+	
+	CompositeCommandCategory.AddCustomBuilder( CommandArrayBuilder, false );
+	
+}
+
+void FBaseCompositeCommandCustomization::CustomizeDetails(const TSharedPtr<IDetailLayoutBuilder>& DetailBuilder)
+{
+	CustomizeDetails(*DetailBuilder);
+}
+
+void FBaseCompositeCommandCustomization::GenerateArrayChildren(TSharedRef<IPropertyHandle> PropertyHandle,
+	int32 ArrayIndex, IDetailChildrenBuilder& ChildrenBuilder)
+{
+	UObject* CommandAsUObject=nullptr;
+	PropertyHandle->GetValue(CommandAsUObject);
+	UUTBBaseCommand* Command=Cast<UUTBBaseCommand>(CommandAsUObject);
+	TSharedPtr<STextBlock> TextBlock=SNew(STextBlock)
+			.Text(TAttribute<FText>::Create([Command]()
+			{
+				if (IsValid(Command))
+				{
+					return FText::FromString(Command->Name);
+				}
+				return FText::FromString("Drag command");
+			}));
+	ChildrenBuilder.AddCustomRow(FText::FromString("Command"))
+	.ValueWidget
+	[
+		
+		SNew(SDropTarget)
+		.OnAllowDrop(SDropTarget::FVerifyDrag::CreateLambda([](TSharedPtr<FDragDropOperation> op)
+		{
+			if (op->IsOfType<FCommandDragDropOperation>())
+			{
+				TSharedPtr<FCommandDragDropOperation> CommandOp=StaticCastSharedPtr<FCommandDragDropOperation>(op);
+				if (!CommandOp.IsValid())
+				{
+					return false;
+				}
+					
+				if (CommandOp->Command->GetClass()->GetDefaultObject()==CommandOp->Command)
+				{
+					return false;
+				}
+				return true;
+			}
+			return false;
+			
+		}))
+		.OnDropped(
+			FOnDrop::CreateLambda([PropertyHandle, TextBlock](const FGeometry&,const FDragDropEvent& Event)
+			{
+				TSharedPtr<FCommandDragDropOperation> Op=Event.GetOperationAs<FCommandDragDropOperation>();
+
+				
+				
+				if (Op.IsValid())
+				{
+					PropertyHandle->SetValue(Op->Command.Get());
+					PropertyHandle->NotifyPostChange(EPropertyChangeType::ValueSet);
+					
+					if (IsValid(Op->Command.Get()))
+					{
+						TextBlock->SetText(FText::FromString(Op->Command->Name));
+					}else
+					{
+						TextBlock->SetText(FText::FromString("Drag command"));
+					}
+					
+					return FReply::Handled().EndDragDrop();
+				}
+				return FReply::Unhandled();
+		}))
+		[
+			SNew(SHorizontalBox)
+			+SHorizontalBox::Slot()
+			.HAlign(HAlign_Left)
+			.VAlign(VAlign_Center)
+			.AutoWidth()
+			[
+				TextBlock.ToSharedRef()
+			]
+			+SHorizontalBox::Slot()
+			.HAlign(HAlign_Fill)
+			[
+				SNew(SSpacer)
+			]
+			+SHorizontalBox::Slot()
+			.HAlign(HAlign_Right)
+			.AutoWidth()
+			[
+				PropertyHandle->CreateDefaultPropertyButtonWidgets()
+				
+			]
+
+			
+			
+		]
+	]
+	.NameWidget
+	[
+	PropertyHandle->CreatePropertyNameWidget()
+	]
+	;
+}
diff --git a/Engine/Plugins/Experimental/UserToolBoxBasicCommand/Source/UserToolBoxBasicCommand/Private/Customization/BaseCompositeCommandCustomization.h b/Engine/Plugins/Experimental/UserToolBoxBasicCommand/Source/UserToolBoxBasicCommand/Private/Customization/BaseCompositeCommandCustomization.h
new file mode 100644
--- /dev/null
+++ b/Engine/Plugins/Experimental/UserToolBoxBasicCommand/Source/UserToolBoxBasicCommand/Private/Customization/BaseCompositeCommandCustomization.h
@@ -0,0 +1,31 @@
+﻿// Copyright Epic Games, Inc. All Rights Reserved.
+
+#pragma once
+
+#include "CoreMinimal.h"
+#include "IDetailCustomization.h"
+#include "PropertyHandle.h"
+
+/**
+ * 
+ */
+
+
+
+
+class FBaseCompositeCommandCustomization : public IDetailCustomization
+{
+public:
+	static TSharedRef<IDetailCustomization> MakeInstance()
+	{
+		return MakeShareable(new FBaseCompositeCommandCustomization());
+	}
+	
+
+	virtual void CustomizeDetails(IDetailLayoutBuilder& DetailBuilder) override;
+	virtual void CustomizeDetails(const TSharedPtr<IDetailLayoutBuilder>& DetailBuilder) override;
+	void GenerateArrayChildren(TSharedRef<IPropertyHandle> PropertyHandle, int32 ArrayIndex, IDetailChildrenBuilder& ChildrenBuilder);
+};
+
+
+
diff --git a/Engine/Plugins/Experimental/UserToolBoxBasicCommand/Source/UserToolBoxBasicCommand/Private/Customization/BindableActionCustomization.cpp b/Engine/Plugins/Experimental/UserToolBoxBasicCommand/Source/UserToolBoxBasicCommand/Private/Customization/BindableActionCustomization.cpp
new file mode 100644
--- /dev/null
+++ b/Engine/Plugins/Experimental/UserToolBoxBasicCommand/Source/UserToolBoxBasicCommand/Private/Customization/BindableActionCustomization.cpp
@@ -0,0 +1,168 @@
+﻿// Copyright Epic Games, Inc. All Rights Reserved.
+
+
+#include "BindableActionCustomization.h"
+#include "DetailCategoryBuilder.h"
+#include "DetailLayoutBuilder.h"
+#include "DetailWidgetRow.h"
+#include "ExecuteBindableAction.h"
+#include "IDetailChildrenBuilder.h"
+#include "IHeadMountedDisplay.h"
+#include "UTBBaseCommand.h"
+#include "UTBBaseTab.h"
+#include "Framework/Commands/InputBindingManager.h"
+#include "Widgets/Input/SComboBox.h"
+#include "Widgets/Text/STextBlock.h"
+
+
+
+
+
+
+FBindableActionPropertyCustomization::~FBindableActionPropertyCustomization()
+{
+}
+
+void FBindableActionPropertyCustomization::CustomizeHeader(TSharedRef<IPropertyHandle> PropertyHandle,
+	FDetailWidgetRow& HeaderRow, IPropertyTypeCustomizationUtils& CustomizationUtils)
+{
+	BindingContexts.Empty();
+	StructPropertyHandle=PropertyHandle;
+	FInputBindingManager::Get().GetKnownInputContexts(BindingContexts);
+	void* ValuePtr;
+	StructPropertyHandle->GetValueData(ValuePtr);
+	FBindableActionInfo* PropertyValue = reinterpret_cast<FBindableActionInfo*>(ValuePtr);
+	TSharedPtr<IPropertyHandle> ParentPropertyHandle=StructPropertyHandle->GetParentHandle();
+	if (PropertyValue)
+	{
+		if (PropertyValue->Context!="") 
+		{
+			for (TSharedPtr<FBindingContext> Context:BindingContexts)
+			{
+				if (Context->GetContextName()==FName(PropertyValue->Context))
+				{
+					SelectedContext=Context;
+					FInputBindingManager::Get().GetCommandInfosFromContext(SelectedContext->GetContextName(),CurrentCommandInfos);
+					CurrentCommand=TSharedPtr<FUICommandInfo>();
+					for (TSharedPtr<FUICommandInfo> CommandInfo:CurrentCommandInfos)
+					{
+						if (CommandInfo->GetCommandName()==FName(PropertyValue->CommandName))
+						{
+							CurrentCommand=CommandInfo;
+							break;
+						}
+					}	
+				}
+			}
+		}
+	}
+	HeaderRow.NameContent()
+	[
+		PropertyHandle->CreatePropertyNameWidget()
+	]
+	.ValueContent()
+	.VAlign(VAlign_Center)
+	[
+		SNew(SVerticalBox)
+		+SVerticalBox::Slot()
+		[
+			SNew(SComboBox<TSharedPtr<FBindingContext> >) // pose list combo
+			.OptionsSource(&BindingContexts)
+			.OnGenerateWidget(this, &FBindableActionPropertyCustomization::OnGenerateContextComboWidget)
+			.OnSelectionChanged(this, &FBindableActionPropertyCustomization::OnSelectContext)
+			[
+				SNew(STextBlock).Text_Lambda([this]()
+				{
+					if (SelectedContext.IsValid())
+					{
+						return FText::FromName(SelectedContext->GetContextName());
+					}
+					return FText::FromString("No Context");
+				})
+			]
+		]
+		+SVerticalBox::Slot()
+		[
+			SNew(SComboBox<TSharedPtr<FUICommandInfo> >) // pose list combo
+			.OptionsSource(&CurrentCommandInfos)
+			.OnGenerateWidget(this, &FBindableActionPropertyCustomization::OnGenerateCommandComboWidget)
+			.OnSelectionChanged(this, &FBindableActionPropertyCustomization::OnSelectCommand)
+			[
+				SNew(STextBlock).Text_Lambda([this]()
+				{
+					if (CurrentCommand.IsValid())
+					{
+						return FText::FromName(CurrentCommand->GetCommandName());
+					}
+					return FText::FromString("No Command");
+				})
+			]
+		]
+		
+	];
+	
+	UE_LOG(LogTemp,Display,TEXT("Header: %s"),*PropertyHandle->GetProperty()->GetName());
+	
+}
+
+void FBindableActionPropertyCustomization::RefreshCommandList()
+{
+	if (SelectedContext.IsValid())
+	{
+		FInputBindingManager::Get().GetCommandInfosFromContext(SelectedContext->GetContextName(),CurrentCommandInfos);
+	}
+	else
+	{
+		CurrentCommandInfos.Empty();
+	}
+}
+TSharedRef<SWidget> FBindableActionPropertyCustomization::OnGenerateContextComboWidget(TSharedPtr<FBindingContext> Item) const
+{
+	return SNew(STextBlock).Text(FText::FromName(Item->GetContextName()));
+}
+
+void FBindableActionPropertyCustomization::OnSelectContext(TSharedPtr<FBindingContext> Item, ESelectInfo::Type SelectInfo)
+{
+	if (Item.IsValid())
+	{
+		SelectedContext = Item;
+		RefreshCommandList();
+		if (CurrentCommand.IsValid())
+		{
+			if (!CurrentCommandInfos.Contains(CurrentCommand))
+			{
+				CurrentCommand.Reset();
+			}
+		}
+	}
+}
+
+void FBindableActionPropertyCustomization::OnSelectCommand(TSharedPtr<FUICommandInfo> Item, ESelectInfo::Type SelectInfo)
+{
+	if (Item.IsValid())
+	{
+		CurrentCommand=Item;
+		if (StructPropertyHandle.IsValid())
+		{
+			FBindableActionInfo Info;
+			Info.Context=SelectedContext->GetContextName().ToString();
+			Info.CommandName=CurrentCommand->GetCommandName().ToString();
+			FStructProperty* StructProperty = CastFieldChecked<FStructProperty>(StructPropertyHandle->GetProperty());
+			FString TextValue;
+			StructProperty->Struct->ExportText(TextValue, &Info, nullptr, nullptr, EPropertyPortFlags::PPF_None, nullptr);
+			StructPropertyHandle->SetValueFromFormattedString(TextValue, EPropertyValueSetFlags::DefaultFlags);
+		}
+	}
+}
+
+TSharedRef<SWidget> FBindableActionPropertyCustomization::OnGenerateCommandComboWidget(
+	TSharedPtr<FUICommandInfo> Item) const
+{
+	return SNew(STextBlock).Text(FText::FromName(Item->GetCommandName()));
+}
+
+void FBindableActionPropertyCustomization::CustomizeChildren(TSharedRef<IPropertyHandle> PropertyHandle,
+                                                      IDetailChildrenBuilder& ChildBuilder, IPropertyTypeCustomizationUtils& CustomizationUtils)
+{
+	
+}
diff --git a/Engine/Plugins/Experimental/UserToolBoxBasicCommand/Source/UserToolBoxBasicCommand/Private/Customization/BindableActionCustomization.h b/Engine/Plugins/Experimental/UserToolBoxBasicCommand/Source/UserToolBoxBasicCommand/Private/Customization/BindableActionCustomization.h
new file mode 100644
--- /dev/null
+++ b/Engine/Plugins/Experimental/UserToolBoxBasicCommand/Source/UserToolBoxBasicCommand/Private/Customization/BindableActionCustomization.h
@@ -0,0 +1,47 @@
+﻿// Copyright Epic Games, Inc. All Rights Reserved.
+
+#pragma once
+
+#include "CoreMinimal.h"
+#include "ExecuteBindableAction.h"
+#include "IDetailCustomization.h"
+#include "IPropertyTypeCustomization.h"
+#include "Framework/Commands/UICommandInfo.h"
+/**
+ * 
+ */
+
+
+
+
+class FBindableActionPropertyCustomization : public IPropertyTypeCustomization
+{
+public:
+	static TSharedRef<IPropertyTypeCustomization> MakeInstance()
+	{
+		return MakeShareable(new FBindableActionPropertyCustomization());
+	}
+	virtual ~FBindableActionPropertyCustomization() override;
+	virtual void CustomizeHeader(TSharedRef<IPropertyHandle> PropertyHandle, FDetailWidgetRow& HeaderRow,
+		IPropertyTypeCustomizationUtils& CustomizationUtils) override;
+	virtual void CustomizeChildren(TSharedRef<IPropertyHandle> PropertyHandle, IDetailChildrenBuilder& ChildBuilder,
+	                               IPropertyTypeCustomizationUtils& CustomizationUtils) override;
+	
+	TSharedRef<SWidget> OnGenerateContextComboWidget(TSharedPtr<FBindingContext> Item) const;
+	void OnSelectContext(TSharedPtr<FBindingContext> Item, ESelectInfo::Type SelectInfo);
+	TSharedRef<SWidget> OnGenerateCommandComboWidget(TSharedPtr<FUICommandInfo> Item) const;
+	void OnSelectCommand(TSharedPtr<FUICommandInfo> Item, ESelectInfo::Type SelectInfo);
+private:
+	TSharedPtr<FBindingContext> SelectedContext;
+	TSharedPtr<FUICommandInfo> CurrentCommand;
+
+	TArray<TSharedPtr<FBindingContext>> BindingContexts;
+	TArray<TSharedPtr<FUICommandInfo>> CurrentCommandInfos;
+
+	void RefreshCommandList();
+	TSharedPtr<IPropertyHandle> StructPropertyHandle;
+	
+};
+
+
+
diff --git a/Engine/Plugins/Experimental/UserToolBoxBasicCommand/Source/UserToolBoxBasicCommand/Private/EngineCommand.cpp b/Engine/Plugins/Experimental/UserToolBoxBasicCommand/Source/UserToolBoxBasicCommand/Private/EngineCommand.cpp
new file mode 100644
--- /dev/null
+++ b/Engine/Plugins/Experimental/UserToolBoxBasicCommand/Source/UserToolBoxBasicCommand/Private/EngineCommand.cpp
@@ -0,0 +1,10 @@
+// Copyright Epic Games, Inc. All Rights Reserved.
+
+
+#include "EngineCommand.h"
+#include "Engine/Engine.h"
+#include "Engine/World.h"
+  void UEngineCommand::Execute()
+ {
+ 	GEngine->Exec(GWorld, *Command);
+ }
\ No newline at end of file
diff --git a/Engine/Plugins/Experimental/UserToolBoxBasicCommand/Source/UserToolBoxBasicCommand/Private/ExecuteBindableAction.cpp b/Engine/Plugins/Experimental/UserToolBoxBasicCommand/Source/UserToolBoxBasicCommand/Private/ExecuteBindableAction.cpp
new file mode 100644
--- /dev/null
+++ b/Engine/Plugins/Experimental/UserToolBoxBasicCommand/Source/UserToolBoxBasicCommand/Private/ExecuteBindableAction.cpp
@@ -0,0 +1,58 @@
+// Copyright Epic Games, Inc. All Rights Reserved.
+
+
+#include "ExecuteBindableAction.h"
+
+#include "UserToolBoxBasicCommand.h"
+#include "Framework/Commands/UICommandInfo.h"
+#include "Framework/Commands/InputBindingManager.h"
+#include "Framework/Application/SlateApplication.h"
+UExecuteBindableAction::UExecuteBindableAction()
+{
+	Name="Exec Bindable Action";
+	Tooltip="Execute a input bindable action";
+	Category="Utility";
+}
+
+void UExecuteBindableAction::Execute()
+{
+	if (CommandInfo. Context.IsEmpty())
+	{
+		TArray<TSharedPtr<FBindingContext>> Contexts;
+		FInputBindingManager::Get().GetKnownInputContexts(Contexts);
+		for (TSharedPtr<FBindingContext> CurContext:Contexts)
+		{
+			UE_LOG(LogUserToolBoxBasicCommand,Display,TEXT("%s"),*CurContext->GetContextName().ToString());
+		}
+		return;
+	}
+	TSharedPtr<FUICommandInfo> Command=FInputBindingManager::Get().FindCommandInContext(FName(CommandInfo.Context),FName(CommandInfo.CommandName));
+	if (!Command.IsValid())
+	{
+		UE_LOG(LogUserToolBoxBasicCommand,Warning,TEXT("%s is not available, try these ones:"),*CommandInfo.CommandName);
+		TArray<TSharedPtr<FUICommandInfo>> CommandInfos;
+		FInputBindingManager::Get().GetCommandInfosFromContext(FName(CommandInfo.Context),CommandInfos);
+		for (TSharedPtr<FUICommandInfo> CurrentCommandInfo:CommandInfos)
+		{
+			UE_LOG(LogUserToolBoxBasicCommand,Display,TEXT("%s"),*CurrentCommandInfo->GetCommandName().ToString());
+		}
+	}
+	else
+	{
+		FInputChord OldInputChord;
+		FInputChord TmpChord(FKey(FName("Subtract")),true,false,false,false);
+		
+		
+		FInputBindingManager::Get().GetUserDefinedChord(FName(CommandInfo.Context),FName(CommandInfo.CommandName),EMultipleKeyBindingIndex::Secondary,OldInputChord);
+		Command->SetActiveChord(TmpChord,EMultipleKeyBindingIndex::Secondary);
+		FInputBindingManager::Get().NotifyActiveChordChanged(*Command.Get(),EMultipleKeyBindingIndex::Secondary);
+		FModifierKeysState Modifier(true,false,false,false,false,false,false,false,false);
+		
+		FKeyEvent KeyEvent(FKey(FName("Subtract")),Modifier,0,false,45,109);
+		FSlateApplication::Get().ProcessKeyDownEvent(KeyEvent);
+		Command->SetActiveChord(OldInputChord,EMultipleKeyBindingIndex::Secondary);
+		FInputBindingManager::Get().NotifyActiveChordChanged(*Command.Get(),EMultipleKeyBindingIndex::Secondary);
+		
+	}
+	
+}
diff --git a/Engine/Plugins/Experimental/UserToolBoxBasicCommand/Source/UserToolBoxBasicCommand/Private/ExecutePythonScript.cpp b/Engine/Plugins/Experimental/UserToolBoxBasicCommand/Source/UserToolBoxBasicCommand/Private/ExecutePythonScript.cpp
new file mode 100644
--- /dev/null
+++ b/Engine/Plugins/Experimental/UserToolBoxBasicCommand/Source/UserToolBoxBasicCommand/Private/ExecutePythonScript.cpp
@@ -0,0 +1,34 @@
+// Copyright Epic Games, Inc. All Rights Reserved.
+
+
+#include "ExecutePythonScript.h"
+#include "Engine/Engine.h"
+
+UExecutePythonScript::UExecutePythonScript()
+{
+	Name="Exec Python";
+	Tooltip="Execute a python script";
+	Category="Utility";
+}
+
+void UExecutePythonScript::Execute()
+{
+	FString CommandArg;
+	if (!ScriptPath.FilePath.IsEmpty())
+	{
+		CommandArg=FString::Printf(TEXT("py \"%s\""),*ScriptPath.FilePath);
+		if (!Args.IsEmpty())
+		{
+			CommandArg+= " " + FString::Printf(TEXT("\"%s\""),*Args);
+		}
+	}
+	else
+	{
+		if (!Args.IsEmpty())
+		{
+			CommandArg+= "py " + Args;
+		}	
+	}
+	
+	GEngine->Exec(NULL, *CommandArg);
+}
diff --git a/Engine/Plugins/Experimental/UserToolBoxBasicCommand/Source/UserToolBoxBasicCommand/Private/FlipNormals.cpp b/Engine/Plugins/Experimental/UserToolBoxBasicCommand/Source/UserToolBoxBasicCommand/Private/FlipNormals.cpp
new file mode 100644
--- /dev/null
+++ b/Engine/Plugins/Experimental/UserToolBoxBasicCommand/Source/UserToolBoxBasicCommand/Private/FlipNormals.cpp
@@ -0,0 +1,54 @@
+// Copyright Epic Games, Inc. All Rights Reserved.
+
+
+#include "FlipNormals.h"
+
+#include "UserToolBoxBasicCommand.h"
+
+#include "UDynamicMesh.h"
+#include "GeometryScript/MeshAssetFunctions.h"
+#include "EdModeInteractiveToolsContext.h"
+#include "Engine/Selection.h"
+#include "Engine/StaticMeshActor.h"
+#include "GeometryScript/MeshNormalsFunctions.h"
+#include "Editor.h"
+#include "Components/StaticMeshComponent.h"
+
+void UFlipNormals::Execute()
+{
+
+	UDynamicMesh* DynamicMesh=NewObject<UDynamicMesh>();
+	USelection* SelectedActors = GEditor->GetSelectedActors();
+	TArray<AActor*> Actors;
+	SelectedActors->GetSelectedObjects(Actors);
+	FGeometryScriptCopyMeshFromAssetOptions OptionsFrom;
+	OptionsFrom.bIgnoreRemoveDegenerates=true;
+	FGeometryScriptCopyMeshToAssetOptions OptionsTo;
+	
+	FGeometryScriptMeshReadLOD Lod;
+	FGeometryScriptMeshWriteLOD TargetLod;
+	EGeometryScriptOutcomePins Result;
+	for (AActor* Actor:Actors)
+	{
+		TArray<UStaticMeshComponent*> SelectedStaticMeshComponents;
+		Actor->GetComponents(SelectedStaticMeshComponents);
+		for (UStaticMeshComponent* SMC:SelectedStaticMeshComponents)
+		{
+			UGeometryScriptLibrary_StaticMeshFunctions::CopyMeshFromStaticMesh(SMC->GetStaticMesh(),DynamicMesh,OptionsFrom,Lod,Result);
+			if (Result==EGeometryScriptOutcomePins::Success)
+			{
+				UGeometryScriptLibrary_MeshNormalsFunctions::FlipNormals(DynamicMesh);
+				AddObjectToTransaction(SMC->GetStaticMesh());
+				UGeometryScriptLibrary_StaticMeshFunctions::CopyMeshToStaticMesh(DynamicMesh,SMC->GetStaticMesh(),OptionsTo,TargetLod,Result);
+			}
+		}
+		
+	}
+}
+
+UFlipNormals::UFlipNormals()
+{
+	Name="Flip normals";
+	Tooltip="Flip normals on each mesh used by selected actor";
+	Category="Mesh";
+}
diff --git a/Engine/Plugins/Experimental/UserToolBoxBasicCommand/Source/UserToolBoxBasicCommand/Private/IsolateSelection.cpp b/Engine/Plugins/Experimental/UserToolBoxBasicCommand/Source/UserToolBoxBasicCommand/Private/IsolateSelection.cpp
new file mode 100644
--- /dev/null
+++ b/Engine/Plugins/Experimental/UserToolBoxBasicCommand/Source/UserToolBoxBasicCommand/Private/IsolateSelection.cpp
@@ -0,0 +1,127 @@
+// Copyright Epic Games, Inc. All Rights Reserved.
+
+
+ #include "IsolateSelection.h"
+
+#include "ActorEditorUtils.h"
+#include "EngineUtils.h"
+#include "Engine/Selection.h"
+#include "Editor.h"
+#include "Engine/StaticMeshActor.h"
+#include "Engine/Classes/Components/MeshComponent.h"
+
+UIsolateSelection::UIsolateSelection()
+{
+	Name="Isolate";
+	Tooltip="Hide all actors except the selected ones";
+	Category="Scene";
+	bIsTransaction=true;
+}
+
+void UIsolateSelection::Execute()
+{
+	UWorld* World = GEngine->GetWorldContexts()[0].World();
+
+	USelection* SelectedActors = GEditor->GetSelectedActors();
+
+	bool IsIsolated = true;
+	UClass* ClassToHide=bShouldOnlyAffectStaticMeshActor?AStaticMeshActor::StaticClass():AActor::StaticClass();
+	for (TActorIterator<AActor> ActorIterator(World, ClassToHide); ActorIterator; ++ActorIterator)
+	{
+		AActor* Actor = CastChecked<AActor>(*ActorIterator);
+
+		if (Actor == nullptr
+			|| FActorEditorUtils::IsABuilderBrush(Actor)
+			|| Actor->IsTemporarilyHiddenInEditor()
+			)
+		{
+			continue;
+		}
+
+		TArray<UMeshComponent*> Components;
+		Actor->GetComponents(Components);
+
+		if (Components.Num() <= 0)
+		{
+			continue;
+		}
+
+		bool IsSelected = false;
+		for (FSelectionIterator SelectedActorIterator(*SelectedActors); SelectedActorIterator; ++SelectedActorIterator)
+		{
+			AActor* SelectedActor = CastChecked<AActor>(*SelectedActorIterator);
+			if (Actor == SelectedActor)
+			{
+				IsSelected = true;
+				break;
+			}
+		}
+		if (IsSelected)
+		{
+			continue;	
+		}
+		IsIsolated = false;
+	}
+
+	if (IsIsolated)
+	{
+		
+		for (TActorIterator<AActor> ActorIterator(World, ClassToHide); ActorIterator; ++ActorIterator)
+		{
+			AActor* Actor = CastChecked<AActor>(*ActorIterator);
+			if (Actor == nullptr || FActorEditorUtils::IsABuilderBrush(Actor))
+			{
+				continue;
+			}
+
+			TArray<UMeshComponent*> Components;
+			Actor->GetComponents(Components);
+
+			if (Components.Num() <= 0)
+			{
+				continue;	
+			}
+
+			Actor->Modify();
+			Actor->SetIsTemporarilyHiddenInEditor(false);
+		}
+	}
+	else
+	{
+		//hide function
+		{
+
+			// Iterate through all of the actors and unhide them
+			for (TActorIterator<AActor> ActorIterator(World, ClassToHide); ActorIterator; ++ActorIterator)
+			{
+				AActor* Actor = *ActorIterator;
+				if (!FActorEditorUtils::IsABuilderBrush(Actor))
+				{
+					TArray<UMeshComponent*> Components;
+					Actor->GetComponents(Components);
+					// Save the actor to the transaction buffer to support undo/redo, but do
+					// not call Modify, as we do not want to dirty the actor's package and
+					// we're only editing temporary, transient values
+					Actor->Modify();
+					if (Components.Num() > 0)
+					{
+						Actor->SetIsTemporarilyHiddenInEditor(true);
+					}
+				}
+			}
+			for (FSelectionIterator It(*SelectedActors); It; ++It)
+			{
+				AActor* Actor = static_cast<AActor*>(*It);
+				checkSlow(Actor->IsA(AActor::StaticClass()));
+
+				// Don't consider already hidden actors or the builder brush
+				if (!FActorEditorUtils::IsABuilderBrush(Actor))
+				{
+					Actor->Modify();
+					Actor->SetIsTemporarilyHiddenInEditor(false);
+				}
+			}
+		}
+	}
+	
+}
diff --git a/Engine/Plugins/Experimental/UserToolBoxBasicCommand/Source/UserToolBoxBasicCommand/Private/Merge.cpp b/Engine/Plugins/Experimental/UserToolBoxBasicCommand/Source/UserToolBoxBasicCommand/Private/Merge.cpp
new file mode 100644
--- /dev/null
+++ b/Engine/Plugins/Experimental/UserToolBoxBasicCommand/Source/UserToolBoxBasicCommand/Private/Merge.cpp
@@ -0,0 +1,83 @@
+// Copyright Epic Games, Inc. All Rights Reserved.
+
+#include "Engine/Selection.h"
+#include "EditorAssetLibrary.h"
+#include "EditorLevelLibrary.h"
+#include "Engine/StaticMeshActor.h"
+#include "Editor.h"
+
+#include "StaticMeshEditorSubsystem.h"
+#include <functional>
+#include "UserToolBoxBasicCommand.h"
+#include "MergeCommand.h"
+#include "Framework/Docking/TabManager.h"
+
+UMerge::UMerge()
+{
+	Name="Merge";
+	Tooltip="Merge selected actor together";
+	Category="Mesh";
+}
+
+void UMerge::Execute()
+{
+	// handle advanced mode
+	if (Advanced)
+	{
+		FGlobalTabmanager::Get()->TryInvokeTab(FName("MergeActors"));
+		return;
+	}
+
+	GEditor->BeginTransaction(FText::FromString("DeleteSelectedActors"));
+
+	TArray<AStaticMeshActor*> ActorsToMerge;
+	USelection* SelectedActors = GEditor->GetSelectedActors();
+	for (FSelectionIterator ActorIterator(*SelectedActors); ActorIterator; ++ActorIterator)
+	{
+		AStaticMeshActor* StaticMeshActor = CastChecked<AStaticMeshActor>(*ActorIterator);
+		if (StaticMeshActor == nullptr)
+		{
+			continue;	
+		}
+		ActorsToMerge.Push(StaticMeshActor);
+	}
+
+	FString NewName = "SM_MERGED_" + ActorsToMerge[0]->GetActorLabel();
+
+	if (UEditorAssetLibrary::DoesAssetExist("/Game/" + NewName))
+	{
+		NewName += "_";
+		int i = 0;
+		while (UEditorAssetLibrary::DoesAssetExist("/Game/" + NewName + FString::FromInt(i)))
+		{
+			i++;
+		}
+		NewName += FString::FromInt(i);
+	}
+
+	UStaticMeshEditorSubsystem* StaticMeshEditorSubsystem = GEditor->GetEditorSubsystem<UStaticMeshEditorSubsystem>();
+	if (StaticMeshEditorSubsystem == nullptr)
+	{
+		UE_LOG(LogUserToolBoxBasicCommand, Warning, TEXT("Could no retrieve StaticMeshEditorSubsystem"));
+		return;
+	}
+	FMergeStaticMeshActorsOptions MergeOptions;
+	auto MergeStaticMeshActors = std::bind(
+		&UStaticMeshEditorSubsystem::MergeStaticMeshActors,
+		StaticMeshEditorSubsystem,
+		std::placeholders::_1,
+		std::placeholders::_2,
+		std::placeholders::_3
+	);
+
+
+	MergeOptions.BasePackageName = "/Game/" + NewName;
+	MergeOptions.NewActorLabel = NewName;
+	AStaticMeshActor* NewActor;
+	
+	MergeStaticMeshActors(ActorsToMerge, MergeOptions, NewActor);
+
+	
+
+	return;
+}
diff --git a/Engine/Plugins/Experimental/UserToolBoxBasicCommand/Source/UserToolBoxBasicCommand/Private/MirrorActorCommand.cpp b/Engine/Plugins/Experimental/UserToolBoxBasicCommand/Source/UserToolBoxBasicCommand/Private/MirrorActorCommand.cpp
new file mode 100644
--- /dev/null
+++ b/Engine/Plugins/Experimental/UserToolBoxBasicCommand/Source/UserToolBoxBasicCommand/Private/MirrorActorCommand.cpp
@@ -0,0 +1,42 @@
+// Copyright Epic Games, Inc. All Rights Reserved.
+
+
+#include "MirrorActorCommand.h"
+#include "Engine/Selection.h"
+#include "Editor.h"
+UMirrorActorCommand::UMirrorActorCommand()
+{
+	Name="Mirror actor";
+	Tooltip="Create a mirror actor depending of the axis selectionned";
+	Category="Actor";
+	bIsTransaction=true;
+}
+
+void UMirrorActorCommand::Execute()
+{
+
+	
+	USelection* SelectedActors = GEditor->GetSelectedActors();
+	TArray<UObject*> ObjectsSelected;
+	SelectedActors->GetSelectedObjects(ObjectsSelected);
+	FVector Mirror(XAxis?-1.0f:1.0f,YAxis?-1.0f:1.0f,ZAxis?-1.0f:1.0f);
+	if (!(XAxis || YAxis || ZAxis))
+	{
+		Mirror=FVector(-1.0,-1.0,-1.0);
+	}
+	for (UObject* Selected : ObjectsSelected)
+	{
+		AActor* InputActor = Cast<AActor>(Selected);
+
+		//newly spawned actor identical to the input actor
+		UWorld* World = InputActor->GetWorld();
+		FActorSpawnParameters params;
+		params.Template = InputActor;
+		UClass* ItemClass = InputActor->GetClass();
+		AddObjectToTransaction(World);
+		AActor* const SpawnedActor = World->SpawnActor<AActor>(ItemClass, params);
+		SpawnedActor->SetActorRelativeScale3D(SpawnedActor->GetActorRelativeScale3D()*Mirror);
+		SpawnedActor->SetActorLabel("Mirrored_"+InputActor->GetActorLabel());
+	}
+	return Super::Execute();
+}
diff --git a/Engine/Plugins/Experimental/UserToolBoxBasicCommand/Source/UserToolBoxBasicCommand/Private/PushComponentMaterialIntoMesh.cpp b/Engine/Plugins/Experimental/UserToolBoxBasicCommand/Source/UserToolBoxBasicCommand/Private/PushComponentMaterialIntoMesh.cpp
new file mode 100644
--- /dev/null
+++ b/Engine/Plugins/Experimental/UserToolBoxBasicCommand/Source/UserToolBoxBasicCommand/Private/PushComponentMaterialIntoMesh.cpp
@@ -0,0 +1,56 @@
+// Copyright Epic Games, Inc. All Rights Reserved.
+
+
+#include "PushComponentMaterialIntoMesh.h"
+#include "Editor.h"
+#include "Selection.h"
+#include "Components/StaticMeshComponent.h"
+void UPushComponentMaterialIntoMesh::Execute()
+{
+	USelection* SelectedActors = GEditor->GetSelectedActors();
+	TArray<UObject*> ActorsAsObject;
+	SelectedActors->GetSelectedObjects(AActor::StaticClass(),ActorsAsObject);
+	TArray<UStaticMeshComponent*> SMCs;
+	for (UObject* Object:ActorsAsObject)
+	{
+		AActor* Actor=Cast<AActor>(Object);
+		if (IsValid(Actor))
+		{
+			TArray<UStaticMeshComponent*> Components;
+			Actor->GetComponents(Components);
+			for (UStaticMeshComponent* Component:Components)
+			{
+				UStaticMesh* StaticMesh=Cast<UStaticMesh>(Component->GetStaticMesh());
+				if (!IsValid(StaticMesh))
+				{
+					return;
+				}
+				TArray<UMaterialInterface*> MatFromSMC=Component->GetMaterials();
+				TArray<UMaterialInterface*> MatFromSM;
+				for ( const FStaticMaterial& Material:StaticMesh->GetStaticMaterials())
+				{
+					MatFromSM.Add(Material.MaterialInterface);
+				}
+				if (MatFromSM!=MatFromSMC)
+				{
+					AddObjectToTransaction(StaticMesh);
+					for (int Index=0;Index<MatFromSMC.Num();Index++ )
+					{
+						StaticMesh->SetMaterial(Index,MatFromSMC[Index]);
+					
+					}
+				}
+			}
+		}
+	}
+}
+
+
+
+UPushComponentMaterialIntoMesh::UPushComponentMaterialIntoMesh()
+{
+	Name="Material from component to mesh";
+	Tooltip="for the current selection, push every material on a static mesh component into the static mesh";
+	Category="Actor";
+	bIsTransaction=true;
+}
diff --git a/Engine/Plugins/Experimental/UserToolBoxBasicCommand/Source/UserToolBoxBasicCommand/Private/SelectActorByFilter.cpp b/Engine/Plugins/Experimental/UserToolBoxBasicCommand/Source/UserToolBoxBasicCommand/Private/SelectActorByFilter.cpp
new file mode 100644
--- /dev/null
+++ b/Engine/Plugins/Experimental/UserToolBoxBasicCommand/Source/UserToolBoxBasicCommand/Private/SelectActorByFilter.cpp
@@ -0,0 +1,154 @@
+// Copyright Epic Games, Inc. All Rights Reserved.
+
+
+#include "SelectActorByFilter.h"
+#include "EngineUtils.h"
+#include "Engine/Selection.h"
+#include "ISinglePropertyView.h"
+#include "UserToolBoxSubsystem.h"
+#include "Dialog/SCustomDialog.h"
+#include "Kismet2/KismetEditorUtilities.h"
+#include "Engine/Engine.h"
+#include "Editor.h"
+#include "Modules/ModuleManager.h"
+#include "PropertyEditorModule.h"
+
+#define LOCTEXT_NAMESPACE "UserToolBoxSelectActorByFilter"
+void USelectActorByFilter::Execute()
+{
+	UWorld* World = GEngine->GetWorldContexts()[0].World();
+	USelection* SelectedActors = GEditor->GetSelectedActors();
+
+
+	FPropertyEditorModule& PropertyEditorModule = FModuleManager::GetModuleChecked<FPropertyEditorModule>("PropertyEditor");
+	FDetailsViewArgs Args;
+	Args.bAllowSearch=false;
+	Args.bShowOptions=false;
+	Args.bHideSelectionTip=true;
+	Args.bShowKeyablePropertiesOption=false;
+	Args.bShowObjectLabel=false;
+	Args.bUpdatesFromSelection=false;
+	Args.NameAreaSettings=FDetailsViewArgs::HideNameArea;
+	TSharedPtr<SVerticalBox> VerticalBox; 
+	TSharedRef<SCustomDialog> OptionsWindow =
+		SNew(SCustomDialog)
+		.Title(FText::FromString(Name+" properties"))
+	.Buttons({
+				SCustomDialog::FButton(LOCTEXT("OK", "OK")),
+				SCustomDialog::FButton(LOCTEXT("Cancel", "Cancel"))
+		})
+	.Content()
+	[
+			SNew(SBox)
+		.MinDesiredWidth(400)
+		[
+		SAssignNew(VerticalBox,SVerticalBox)
+
+		]
+		
+	];
+
+	int ShowedFilterCount=0;
+	for (FActorFilterOptions& Filter:FilterStack)
+	{
+		if (Filter.Filter==nullptr)
+		{
+			continue;
+		}
+		if (!Filter.ShowProperties)
+		{
+			continue;
+		}
+		TSharedRef<IDetailsView> DetailsView=PropertyEditorModule.CreateDetailView(Args);
+		DetailsView->SetObject(Filter.Filter);	
+		VerticalBox->AddSlot()
+		[
+			DetailsView
+		];
+		ShowedFilterCount++;
+	}
+	if (ShowedFilterCount>0)
+	{
+		const int32 PressedButtonIdx = OptionsWindow->ShowModal();
+		if (PressedButtonIdx != 0)
+		{
+			
+			return;
+		}	
+	}
+	TArray<AActor*> CurrentSelection;
+	if (ApplyToCurrentSelection)
+	{
+		SelectedActors->GetSelectedObjects(CurrentSelection);
+		
+	}else
+	{
+		for (TActorIterator<AActor> ActorIterator(World, AActor::StaticClass()); ActorIterator; ++ActorIterator)
+		{
+			CurrentSelection.Add(*ActorIterator);
+		}
+	}
+	TArray<AActor*> OriginalSelection=CurrentSelection;
+	for (FActorFilterOptions& Filter:FilterStack)
+	{
+		if (Filter.Filter==nullptr)
+		{
+			continue;
+		}
+			
+		TArray<AActor*>Result=Filter.Filter->FilterImpl(Filter.Source==EActorFilterSource::OriginalSelection?OriginalSelection:CurrentSelection);
+		if (Filter.Rule==EActorFilterRule::Add)
+		{
+			CurrentSelection.Append(Result);
+		}else if (Filter.Rule==EActorFilterRule::Intersect)
+		{
+			TSet<AActor*> ResultSet=TSet<AActor*>(Result);
+			TArray<AActor*> Intersect;
+			for(AActor* Item:CurrentSelection)
+			{
+				if (ResultSet.Contains(Item))
+				{
+					Intersect.AddUnique(Item);
+				}
+			}
+			CurrentSelection=Intersect;
+		}else if (Filter.Rule==EActorFilterRule::Substract)
+		{
+			TSet<AActor*> ResultSet=TSet<AActor*>(Result);
+			TArray<AActor*> Substract;
+			for(AActor* Item:CurrentSelection)
+			{
+				if (!ResultSet.Contains(Item))
+				{
+					Substract.AddUnique(Item);
+				}
+			}
+			CurrentSelection=Substract;
+		}else if (Filter.Rule==EActorFilterRule::Replace)
+		{
+			CurrentSelection=Result;
+		}
+	}
+	SelectedActors->Modify();
+	SelectedActors->BeginBatchSelectOperation();
+	SelectedActors->DeselectAll();
+	for (AActor* Actor:CurrentSelection)
+	{
+		if (Actor!=nullptr)
+		{
+			SelectedActors->Select(Actor,true);
+		}
+	}
+	SelectedActors->EndBatchSelectOperation();
+	GEditor->NoteSelectionChange(true);
+	return Super::Execute();
+}
+
+USelectActorByFilter::USelectActorByFilter()
+{
+	Name="Select actor by filter";
+	Tooltip="Select actor from a stack of filter";
+	Category="Actor";
+	bIsTransaction=true;
+}
+#undef LOCTEXT_NAMESPACE
\ No newline at end of file
diff --git a/Engine/Plugins/Experimental/UserToolBoxBasicCommand/Source/UserToolBoxBasicCommand/Private/SelectActorBySize.cpp b/Engine/Plugins/Experimental/UserToolBoxBasicCommand/Source/UserToolBoxBasicCommand/Private/SelectActorBySize.cpp
new file mode 100644
--- /dev/null
+++ b/Engine/Plugins/Experimental/UserToolBoxBasicCommand/Source/UserToolBoxBasicCommand/Private/SelectActorBySize.cpp
@@ -0,0 +1,64 @@
+// Copyright Epic Games, Inc. All Rights Reserved.
+
+
+#include "SelectActorBySize.h"
+
+
+#include "EngineUtils.h"
+#include "SkeletalDebugRendering.h"
+#include "Engine/Selection.h"
+#include "UserToolBoxBasicCommand.h"
+#include "Editor.h"
+
+USelectActorBySize::USelectActorBySize()
+{
+	Name="Select actor by size";
+	Tooltip="Select the actor depending of their bounding box size";
+	Category="Actor";
+	bIsTransaction=true;
+}
+
+void USelectActorBySize::Execute()
+{
+	UWorld* World = GEngine->GetWorldContexts()[0].World();
+
+	
+	TArray<AActor*> NewSelection;
+	auto FilterPerBoundingBox = [&NewSelection, this](AActor* CurActor)
+	{
+		if (CurActor->GetClass()->HasAnyClassFlags(CLASS_NotPlaceable) )
+		{
+			return;
+		}
+		FVector BoundingBox, Origin;
+		CurActor->GetActorBounds(false, Origin, BoundingBox, true);
+		BoundingBox *= 2;
+		float Size=SizeThreshold;
+		if (BoundingBox.X < SizeThreshold
+			&& BoundingBox.Y < SizeThreshold
+			&& BoundingBox.Z < SizeThreshold
+			)
+		{
+			NewSelection.Add(CurActor);
+		}
+	};
+
+	for (TActorIterator<AActor> It(World, AActor::StaticClass()); It; ++It)
+	{
+		if (*It != nullptr)
+		{
+			FilterPerBoundingBox(*It);
+		}
+	}
+	USelection* SelectedActors = GEditor->GetSelectedActors();
+	AddObjectToTransaction(SelectedActors);
+	SelectedActors->BeginBatchSelectOperation();
+	SelectedActors->DeselectAll();
+	for (AActor* CurActor : NewSelection)
+	{
+		SelectedActors->Select(CurActor);
+	}
+	SelectedActors->EndBatchSelectOperation();
+	GEditor->NoteSelectionChange();
+	
+}
diff --git a/Engine/Plugins/Experimental/UserToolBoxBasicCommand/Source/UserToolBoxBasicCommand/Private/SetHighPrecisionOnMesh.cpp b/Engine/Plugins/Experimental/UserToolBoxBasicCommand/Source/UserToolBoxBasicCommand/Private/SetHighPrecisionOnMesh.cpp
new file mode 100644
--- /dev/null
+++ b/Engine/Plugins/Experimental/UserToolBoxBasicCommand/Source/UserToolBoxBasicCommand/Private/SetHighPrecisionOnMesh.cpp
@@ -0,0 +1,51 @@
+// Copyright Epic Games, Inc. All Rights Reserved.
+
+
+#include "SetHighPrecisionOnMesh.h"
+#include "Engine/Selection.h"
+#include "Engine/StaticMesh.h"
+#include "Editor.h"
+#include "Components/StaticMeshComponent.h"
+
+USetHighPrecisionOnMesh::USetHighPrecisionOnMesh()
+{
+	Name="Set HighPrecision";
+	Tooltip=" Use HighPrecision on mesh used by selected actors";
+	Category="Mesh";
+	bIsTransaction=true;
+}
+
+void USetHighPrecisionOnMesh::Execute()
+{
+	USelection* SelectedActors = GEditor->GetSelectedActors();
+	TArray<AActor*> Actors;
+	TArray<UStaticMesh*> StaticMeshes;
+	SelectedActors->GetSelectedObjects(Actors);
+	for (AActor* Actor:Actors)
+	{
+		TArray<UStaticMeshComponent*> StaticMeshComponents;
+		Actor->GetComponents(StaticMeshComponents);
+		for (UStaticMeshComponent* StaticMeshComponent:StaticMeshComponents)
+		{
+			UStaticMesh* StaticMesh=StaticMeshComponent->GetStaticMesh();
+			if (IsValid(StaticMesh))
+			{
+				StaticMeshes.Add(StaticMesh);
+			}	
+		}
+		
+	}
+	for (UStaticMesh* StaticMesh:StaticMeshes)
+	{
+		AddObjectToTransaction(StaticMesh);
+		for (int32 LODIndex = 0; LODIndex < StaticMesh->GetNumLODs(); LODIndex++)
+		{
+			FMeshBuildSettings& LODBuildSettings = StaticMesh->GetSourceModel(LODIndex).BuildSettings;
+			LODBuildSettings.bUseHighPrecisionTangentBasis=bHighPrecisionTangent;
+			LODBuildSettings.bUseFullPrecisionUVs=bHighPrecisionUV;
+			StaticMesh->PostEditChange();
+			StaticMesh->MarkPackageDirty();
+		}
+	}
+	return;
+}
diff --git a/Engine/Plugins/Experimental/UserToolBoxBasicCommand/Source/UserToolBoxBasicCommand/Private/ShowLayersCommand.cpp b/Engine/Plugins/Experimental/UserToolBoxBasicCommand/Source/UserToolBoxBasicCommand/Private/ShowLayersCommand.cpp
new file mode 100644
--- /dev/null
+++ b/Engine/Plugins/Experimental/UserToolBoxBasicCommand/Source/UserToolBoxBasicCommand/Private/ShowLayersCommand.cpp
@@ -0,0 +1,183 @@
+// Copyright Epic Games, Inc. All Rights Reserved.
+
+
+
+#include "ShowLayersCommand.h"
+#include "Layers/Layer.h"
+#include "Layers/LayersSubsystem.h"
+#include "UserToolBoxBasicCommand.h"
+#include "Editor.h"
+#include "Widgets/Input/SButton.h"
+#include "Widgets/Input/SCheckBox.h"
+#include "Widgets/Text/STextBlock.h"
+UShowLayersCommand::UShowLayersCommand()
+{
+	Name="Show layers";
+	Tooltip="Show specific layers, if isolate, hide the other ones";
+	Category="Scene";
+}
+
+bool UShowLayersCommand::DisplayParameters()
+{
+	ULayersSubsystem* LayersSubsystem = GEditor->GetEditorSubsystem<ULayersSubsystem>();
+
+	if (LayersSubsystem == nullptr)
+	{
+		UE_LOG(LogUserToolBoxBasicCommand, Warning, TEXT("Could no retrieve LayersSubsystem"));
+		return false;
+	}
+	const ISlateStyle& EditorStyle=FAppStyle::Get();
+	TArray<TWeakObjectPtr<ULayer>> AvailableLayers;
+	LayersSubsystem->AddAllLayersTo(AvailableLayers);
+
+	TSharedRef<SWindow> Window = SNew(SWindow)
+		.Title(NSLOCTEXT("UTBAssignToLayer", "SelectLayer", "Select Layer"))
+		.ClientSize(FVector2D(100.f, 30.f * (AvailableLayers.Num()+3)))
+		.IsTopmostWindow(true);
+	bool SelectionValid=false;
+	TSharedRef<SVerticalBox> LayerList = SNew(SVerticalBox)
+	+SVerticalBox::Slot()
+	[
+		SNew(SHorizontalBox)
+		+SHorizontalBox::Slot()
+		[
+			SNew(SButton)
+			.Text(FText::FromString("Ok"))
+			.HAlign(EHorizontalAlignment::HAlign_Center)
+			.VAlign(EVerticalAlignment::VAlign_Center)
+			.ButtonStyle(EditorStyle, "PropertyEditor.AssetComboStyle")
+			.ForegroundColor(EditorStyle.GetColor("PropertyEditor.AssetName.ColorAndOpacity"))
+			.OnClicked(FOnClicked::CreateLambda([&SelectionValid,Window]()
+			{
+				SelectionValid=true;
+				Window->RequestDestroyWindow();
+				return FReply::Handled();
+			}))
+		]
+		+SHorizontalBox::Slot()
+		[
+			SNew(SButton)
+			.Text(FText::FromString("Cancel"))
+			.HAlign(EHorizontalAlignment::HAlign_Center)
+			.VAlign(EVerticalAlignment::VAlign_Center)
+			.ButtonStyle(EditorStyle, "PropertyEditor.AssetComboStyle")
+			.ForegroundColor(EditorStyle.GetColor("PropertyEditor.AssetName.ColorAndOpacity"))
+			.OnClicked(FOnClicked::CreateLambda([&SelectionValid,Window]()
+			{
+				Window->RequestDestroyWindow();
+				return FReply::Handled();
+			}))
+		]
+	]
+	+SVerticalBox::Slot()
+	[
+		SNew(SHorizontalBox)
+		+SHorizontalBox::Slot()
+		.HAlign(EHorizontalAlignment::HAlign_Left)
+		.VAlign(EVerticalAlignment::VAlign_Center)
+		.AutoWidth()
+		[
+			SNew(SCheckBox)
+			.IsChecked_Lambda([this]()
+			{
+				return this->bIsolate?ECheckBoxState::Checked:ECheckBoxState::Unchecked;
+			})
+			.OnCheckStateChanged_Lambda([this](ECheckBoxState NewState)
+			{
+				this->bIsolate=NewState==ECheckBoxState::Checked;
+			})
+		]
+		+SHorizontalBox::Slot()
+		.HAlign(EHorizontalAlignment::HAlign_Left)
+		.VAlign(EVerticalAlignment::VAlign_Center)
+		[
+			SNew(STextBlock)
+			.Text(FText::FromString("Isolate"))
+		]
+	];
+
+	TSet<FName> LayersSelected;
+	for (TWeakObjectPtr<ULayer> Layer : AvailableLayers) {
+		TSharedPtr<SCheckBox> Checkbox;
+		LayerList->AddSlot()
+			[
+				
+			
+				SNew(SButton)
+				.HAlign(EHorizontalAlignment::HAlign_Fill)
+				.VAlign(EVerticalAlignment::VAlign_Fill)
+				.ButtonStyle(EditorStyle, "PropertyEditor.AssetComboStyle")
+				.ForegroundColor(EditorStyle.GetColor("PropertyEditor.AssetName.ColorAndOpacity"))
+				[
+					SNew(SHorizontalBox)
+					
+					+SHorizontalBox::Slot()
+					.HAlign(EHorizontalAlignment::HAlign_Left)
+					.VAlign(EVerticalAlignment::VAlign_Center)
+					[
+						SAssignNew(Checkbox,SCheckBox)
+						.Visibility(EVisibility::HitTestInvisible)
+					]
+					+SHorizontalBox::Slot()
+					.HAlign(EHorizontalAlignment::HAlign_Fill)
+					.VAlign(EVerticalAlignment::VAlign_Center)
+					[
+						SNew(STextBlock)
+						.Text(
+							FText::FromName(Layer->GetLayerName())
+
+							)
+						.Justification(ETextJustify::Left)
+					]
+				]
+				.OnClicked_Lambda([Layer,&LayersSelected, Checkbox]()->FReply
+				{
+					Checkbox->ToggleCheckedState();
+					if (Checkbox->IsChecked())
+					{
+						LayersSelected.Add(Layer->GetLayerName());
+					}
+					else
+					{
+						LayersSelected.Remove(Layer->GetLayerName());
+					}
+					return FReply::Handled();
+				})
+			];
+	}
+
+	Window->SetContent(LayerList);
+
+	GEditor->EditorAddModalWindow(Window);
+	Layers=LayersSelected.Array();
+	return true;
+}
+
+void UShowLayersCommand::Execute()
+{
+	ULayersSubsystem* LayersSubsystem = GEditor->GetEditorSubsystem<ULayersSubsystem>();
+	if (IsValid(LayersSubsystem))
+	{
+		TArray<FName> AllLayersName;
+		TArray<FName> Visible;
+		TArray<FName> Invisible;
+		LayersSubsystem->AddAllLayerNamesTo(AllLayersName);
+		for (FName LayerName:AllLayersName)
+		{
+			if (Layers.Contains(LayerName))
+			{
+				Visible.Add(LayerName);
+				
+			}else
+			{
+				if (bIsolate)
+				{
+					Invisible.Add(LayerName);
+				}
+			}
+		}
+		LayersSubsystem->SetLayersVisibility(Visible,true);
+		LayersSubsystem->SetLayersVisibility(Invisible,false);
+	}
+	return;
+}
diff --git a/Engine/Plugins/Experimental/UserToolBoxBasicCommand/Source/UserToolBoxBasicCommand/Private/TabSpawner.cpp b/Engine/Plugins/Experimental/UserToolBoxBasicCommand/Source/UserToolBoxBasicCommand/Private/TabSpawner.cpp
new file mode 100644
--- /dev/null
+++ b/Engine/Plugins/Experimental/UserToolBoxBasicCommand/Source/UserToolBoxBasicCommand/Private/TabSpawner.cpp
@@ -0,0 +1,16 @@
+// Copyright Epic Games, Inc. All Rights Reserved.
+
+
+#include "TabSpawner.h"
+#include "Framework/Docking/TabManager.h"
+UTabSpawner::UTabSpawner()
+{
+	Name="Tab Spawner";
+	Tooltip="Allow you to spawn any tab by name";
+	Category="Utility";
+}
+
+void UTabSpawner::Execute()
+{
+	FGlobalTabmanager::Get()->TryInvokeTab(TabName);
+}
diff --git a/Engine/Plugins/Experimental/UserToolBoxBasicCommand/Source/UserToolBoxBasicCommand/Private/ToggleCommand.cpp b/Engine/Plugins/Experimental/UserToolBoxBasicCommand/Source/UserToolBoxBasicCommand/Private/ToggleCommand.cpp
new file mode 100644
--- /dev/null
+++ b/Engine/Plugins/Experimental/UserToolBoxBasicCommand/Source/UserToolBoxBasicCommand/Private/ToggleCommand.cpp
@@ -0,0 +1,23 @@
+// Copyright Epic Games, Inc. All Rights Reserved.
+
+
+#include "ToggleCommand.h"
+#include "Engine/Selection.h"
+void UToggleCommand::Execute()
+{
+	if (Commands.IsEmpty())
+	{
+		return;
+	}
+	if (CurrentIndex>=Commands.Num())
+	{
+		CurrentIndex=0;
+	}
+	UUTBBaseCommand* Current=Commands[CurrentIndex];
+	CurrentIndex++;
+	if (!IsValid(Current))
+	{
+		return;
+	}
+	Current->ExecuteCommand();
+}
diff --git a/Engine/Plugins/Experimental/UserToolBoxBasicCommand/Source/UserToolBoxBasicCommand/Private/ToggleCommandInline.cpp b/Engine/Plugins/Experimental/UserToolBoxBasicCommand/Source/UserToolBoxBasicCommand/Private/ToggleCommandInline.cpp
new file mode 100644
--- /dev/null
+++ b/Engine/Plugins/Experimental/UserToolBoxBasicCommand/Source/UserToolBoxBasicCommand/Private/ToggleCommandInline.cpp
@@ -0,0 +1,34 @@
+// Copyright Epic Games, Inc. All Rights Reserved.
+
+
+#include "ToggleCommandInline.h"
+
+
+#include "EditorUtilityBlueprint.h"
+#include "EngineUtils.h"
+#include "UserToolBoxBaseBlueprint.h"
+#include "UTBBaseTab.h"
+#include "Engine/Selection.h"
+#include "UserToolBoxBasicCommand.h"
+
+
+
+void UToggleCommandInline::Execute()
+{
+	if (Commands.IsEmpty())
+	{
+		return;
+	}
+	if (CurrentIndex>=Commands.Num())
+	{
+		CurrentIndex=0;
+	}
+	UUTBBaseCommand* Current=Commands[CurrentIndex];
+	CurrentIndex++;
+	if (!IsValid(Current))
+	{
+		return;
+	}
+	Current->ExecuteCommand();
+
+}
diff --git a/Engine/Plugins/Experimental/UserToolBoxBasicCommand/Source/UserToolBoxBasicCommand/Private/UserToolBoxBasicCommand.cpp b/Engine/Plugins/Experimental/UserToolBoxBasicCommand/Source/UserToolBoxBasicCommand/Private/UserToolBoxBasicCommand.cpp
new file mode 100644
--- /dev/null
+++ b/Engine/Plugins/Experimental/UserToolBoxBasicCommand/Source/UserToolBoxBasicCommand/Private/UserToolBoxBasicCommand.cpp
@@ -0,0 +1,34 @@
+// Copyright Epic Games, Inc. All Rights Reserved.
+
+#include "UserToolBoxBasicCommand.h"
+
+#include "AssetToolsModule.h"
+//#include "ExecuteBindableAction.h"
+//#include "IAssetTools.h"
+
+#include "SelectActorByFilter.h"
+#include "ToggleCommand.h"
+#include "Customization/BaseCompositeCommandCustomization.h"
+#include "Customization/BindableActionCustomization.h"
+
+#include "UObject/CoreRedirects.h"
+#include "PropertyEditorModule.h"
+#define LOCTEXT_NAMESPACE "FUserToolBox_BasicCommandModule"
+DEFINE_LOG_CATEGORY(LogUserToolBoxBasicCommand);
+void FUserToolBox_BasicCommandModule::StartupModule()
+{
+
+	FPropertyEditorModule& PropertyEditorModule = FModuleManager::Get().LoadModuleChecked<FPropertyEditorModule>("PropertyEditor");
+	PropertyEditorModule.RegisterCustomClassLayout(UBaseCompositeCommand::StaticClass()->GetFName(),FOnGetDetailCustomizationInstance::CreateStatic(&FBaseCompositeCommandCustomization::MakeInstance));
+	PropertyEditorModule.RegisterCustomPropertyTypeLayout(FBindableActionInfo::StaticStruct()->GetFName(),FOnGetPropertyTypeCustomizationInstance::CreateStatic(&FBindableActionPropertyCustomization::MakeInstance));
+}
+
+void FUserToolBox_BasicCommandModule::ShutdownModule()
+{
+	// This function may be called during shutdown to clean up your module.  For modules that support dynamic reloading,
+	// we call this function before unloading the module.
+}
+
+#undef LOCTEXT_NAMESPACE
+
+IMPLEMENT_MODULE(FUserToolBox_BasicCommandModule, UserToolBox_BasicCommand)
\ No newline at end of file
diff --git a/Engine/Plugins/Experimental/UserToolBoxBasicCommand/Source/UserToolBoxBasicCommand/Private/ZoomAll.cpp b/Engine/Plugins/Experimental/UserToolBoxBasicCommand/Source/UserToolBoxBasicCommand/Private/ZoomAll.cpp
new file mode 100644
--- /dev/null
+++ b/Engine/Plugins/Experimental/UserToolBoxBasicCommand/Source/UserToolBoxBasicCommand/Private/ZoomAll.cpp
@@ -0,0 +1,36 @@
+// Copyright Epic Games, Inc. All Rights Reserved.
+
+
+#include "ZoomAll.h"
+
+#include "GameFramework/Actor.h"
+#include "EngineUtils.h"
+#include "EditorViewportClient.h"
+
+UZoomAll::UZoomAll()
+{
+	Name="Zoom All";
+	Tooltip="Zoom the camera to see the whole scene";
+	Category="Viewport";
+}
+
+void UZoomAll::Execute()
+{
+	UWorld* World = GEngine->GetWorldContexts()[0].World();
+
+	TActorIterator<AActor> Actor(World, AActor::StaticClass());
+	FBox BoundingBox = Actor->GetComponentsBoundingBox();
+
+	for (; Actor; ++Actor)
+	{
+		if (*Actor != nullptr)
+		{
+			BoundingBox += Actor->GetComponentsBoundingBox();
+		}
+	}
+
+	FEditorViewportClient* EditorViewportClient = static_cast<FEditorViewportClient*>(GEditor->GetActiveViewport()->GetClient());
+	EditorViewportClient->FocusViewportOnBox(BoundingBox);
+
+	
+}
diff --git a/Engine/Plugins/Experimental/UserToolBoxBasicCommand/Source/UserToolBoxBasicCommand/Public/ActorFilter/BaseActorFilter.h b/Engine/Plugins/Experimental/UserToolBoxBasicCommand/Source/UserToolBoxBasicCommand/Public/ActorFilter/BaseActorFilter.h
new file mode 100644
--- /dev/null
+++ b/Engine/Plugins/Experimental/UserToolBoxBasicCommand/Source/UserToolBoxBasicCommand/Public/ActorFilter/BaseActorFilter.h
@@ -0,0 +1,139 @@
+// Copyright Epic Games, Inc. All Rights Reserved.
+
+#pragma once
+
+#include "CoreMinimal.h"
+#include "UObject/NoExportTypes.h"
+#include "Components/ActorComponent.h"
+#include "GameFramework/Actor.h"
+#include "BaseActorFilter.generated.h"
+
+/**
+ * 
+ */
+UCLASS(Blueprintable,EditInlineNew, Abstract)
+class USERTOOLBOXBASICCOMMAND_API UBaseActorFilter : public UObject
+{
+	GENERATED_BODY()
+	public:
+	
+	UFUNCTION(BlueprintNativeEvent)
+	TArray<AActor*> Filter(const TArray<AActor*>& Source);
+	UFUNCTION(BlueprintNativeEvent)
+	bool FilterUnit(AActor* Source);
+	UPROPERTY(EditAnywhere,Category="FilterCommand")
+	bool bShouldRunGlobalFilter=true;
+	UPROPERTY(EditAnywhere,Category="FilterCommand")
+	bool bShouldRunUnitFilter=true;
+	virtual TArray<AActor*> FilterImpl(const TArray<AActor*>& Source)
+	{
+		TArray<AActor*> GroupFilter;
+		TArray<AActor*> Result;
+		GroupFilter.Reserve(Source.Num());
+		GroupFilter=bShouldRunGlobalFilter?Filter(Source):Source;
+	
+		Result.Reserve(GroupFilter.Num());
+		if (bShouldRunUnitFilter)
+		{
+			for (AActor* ActorToFilter:GroupFilter)
+			{
+				if (FilterUnit(ActorToFilter))
+				{
+					Result.Add(ActorToFilter);
+				}
+			}
+		}
+		else
+		{
+			Result=GroupFilter;
+		}
+		Result.Shrink();
+		return Result;
+	}
+};
+
+
+UCLASS(Blueprintable)
+class UGetAllDescendants : public UBaseActorFilter
+{
+	GENERATED_BODY()
+public:
+	virtual TArray<AActor*> FilterImpl(const TArray<AActor*>& Source) override;
+};
+
+
+
+UCLASS(Blueprintable)
+class UGetParents : public UBaseActorFilter
+{
+	GENERATED_BODY()
+public:
+	virtual TArray<AActor*> FilterImpl(const TArray<AActor*>& Source) override;
+};
+UCLASS(Blueprintable)
+class UHasAttachedActor : public UBaseActorFilter
+{
+	GENERATED_BODY()
+public:
+	virtual TArray<AActor*> FilterImpl(const TArray<AActor*>& Source) override;
+};
+
+UCLASS(Blueprintable)
+class UHasComponentOfClass : public UBaseActorFilter
+{
+	GENERATED_BODY()
+public:
+	UPROPERTY(EditAnywhere, BlueprintReadWrite, Category="Filter")
+	TSubclassOf<UActorComponent> ComponentClass;
+	virtual TArray<AActor*> FilterImpl(const TArray<AActor*>& Source) override;
+};
+UCLASS(Blueprintable)
+class UHasMetadataByKey : public UBaseActorFilter
+{
+	GENERATED_BODY()
+public:
+	UPROPERTY(EditAnywhere, BlueprintReadWrite, Category="Filter")
+	FName Key;
+	virtual TArray<AActor*> FilterImpl(const TArray<AActor*>& Source) override;
+};
+UCLASS(Blueprintable)
+class UHasMetadataByKeyAndValue : public UBaseActorFilter
+{
+	GENERATED_BODY()
+public:
+	UPROPERTY(EditAnywhere, BlueprintReadWrite, Category="Filter")
+	FName Key;
+	UPROPERTY(EditAnywhere, BlueprintReadWrite, Category="Filter")
+	FString Value;
+	virtual TArray<AActor*> FilterImpl(const TArray<AActor*>& Source) override;
+};
+UCLASS(Blueprintable)
+class UIsClassOf : public UBaseActorFilter
+{
+	GENERATED_BODY()
+public:
+	UPROPERTY(EditAnywhere, BlueprintReadWrite, Category="Filter")
+	TSubclassOf<AActor>	ActorClass;
+	UPROPERTY(EditAnywhere, BlueprintReadWrite, Category="Filter")
+	bool ChildClass=false;
+	virtual TArray<AActor*> FilterImpl(const TArray<AActor*>& Source) override;
+};
+UCLASS(Blueprintable)
+class UGetNDescendants : public UBaseActorFilter
+{
+	GENERATED_BODY()
+public:
+	UPROPERTY(EditAnywhere, BlueprintReadWrite, Category="Filter")
+	int		N;
+	UPROPERTY(EditAnywhere, BlueprintReadWrite, Category="Filter")
+	bool	AddIntermediaries;
+	virtual TArray<AActor*> FilterImpl(const TArray<AActor*>& Source) override;
+};
+UCLASS(Blueprintable)
+class UHasMetadataByKeyAndValueDropDown : public UBaseActorFilter
+{
+	GENERATED_BODY()
+public:
+	
+	virtual TArray<AActor*> FilterImpl(const TArray<AActor*>& Source) override;
+};
diff --git a/Engine/Plugins/Experimental/UserToolBoxBasicCommand/Source/UserToolBoxBasicCommand/Public/AssignToLayer.h b/Engine/Plugins/Experimental/UserToolBoxBasicCommand/Source/UserToolBoxBasicCommand/Public/AssignToLayer.h
new file mode 100644
--- /dev/null
+++ b/Engine/Plugins/Experimental/UserToolBoxBasicCommand/Source/UserToolBoxBasicCommand/Public/AssignToLayer.h
@@ -0,0 +1,22 @@
+// Copyright Epic Games, Inc. All Rights Reserved.
+
+#pragma once
+
+#include "CoreMinimal.h"
+#include "UTBBaseCommand.h"
+#include "Layers/Layer.h"
+#include "AssignToLayer.generated.h"
+
+/**
+ * 
+ */
+UCLASS()
+class USERTOOLBOXBASICCOMMAND_API UAssignToLayer : public UUTBBaseCommand
+{
+	GENERATED_BODY()
+	public:
+
+	void Execute() override;
+	virtual FName GetLayerName(TWeakObjectPtr<ULayer> Layer);
+	UAssignToLayer();
+};
diff --git a/Engine/Plugins/Experimental/UserToolBoxBasicCommand/Source/UserToolBoxBasicCommand/Public/AssignToLevel.h b/Engine/Plugins/Experimental/UserToolBoxBasicCommand/Source/UserToolBoxBasicCommand/Public/AssignToLevel.h
new file mode 100644
--- /dev/null
+++ b/Engine/Plugins/Experimental/UserToolBoxBasicCommand/Source/UserToolBoxBasicCommand/Public/AssignToLevel.h
@@ -0,0 +1,20 @@
+// Copyright Epic Games, Inc. All Rights Reserved.
+
+#pragma once
+
+#include "CoreMinimal.h"
+#include "UTBBaseCommand.h"
+#include "AssignToLevel.generated.h"
+
+/**
+ * 
+ */
+UCLASS()
+class USERTOOLBOXBASICCOMMAND_API UAssignToLevel : public UUTBBaseCommand
+{
+	GENERATED_BODY()
+	public:
+
+	virtual void Execute() override;
+	UAssignToLevel();
+};
diff --git a/Engine/Plugins/Experimental/UserToolBoxBasicCommand/Source/UserToolBoxBasicCommand/Public/BaseCompositeCommand.h b/Engine/Plugins/Experimental/UserToolBoxBasicCommand/Source/UserToolBoxBasicCommand/Public/BaseCompositeCommand.h
new file mode 100644
--- /dev/null
+++ b/Engine/Plugins/Experimental/UserToolBoxBasicCommand/Source/UserToolBoxBasicCommand/Public/BaseCompositeCommand.h
@@ -0,0 +1,30 @@
+// Copyright Epic Games, Inc. All Rights Reserved.
+
+#pragma once
+
+#include "CoreMinimal.h"
+#include "UTBBaseCommand.h"
+#include "BaseCompositeCommand.generated.h"
+
+/**
+ * 
+ */
+UCLASS(Abstract)
+class USERTOOLBOXBASICCOMMAND_API UBaseCompositeCommand : public UUTBBaseCommand
+{
+	GENERATED_BODY()
+public:
+
+	UPROPERTY(EditAnywhere,BlueprintReadWrite, Category="Composite Commands")
+	TArray<TObjectPtr<UUTBBaseCommand>>		Commands;
+	
+	UBaseCompositeCommand()
+	{
+		Name="Toggle Command";
+		Tooltip="This command let you execute sequentially some already registered command, one at a time (placeholder section is a good place to register these commands)";
+		Category="Builtin";
+	}
+
+	virtual UUTBBaseCommand* CopyCommand(UObject* Owner) const override;
+
+};
diff --git a/Engine/Plugins/Experimental/UserToolBoxBasicCommand/Source/UserToolBoxBasicCommand/Public/BaseCompositeInlineCommand.h b/Engine/Plugins/Experimental/UserToolBoxBasicCommand/Source/UserToolBoxBasicCommand/Public/BaseCompositeInlineCommand.h
new file mode 100644
--- /dev/null
+++ b/Engine/Plugins/Experimental/UserToolBoxBasicCommand/Source/UserToolBoxBasicCommand/Public/BaseCompositeInlineCommand.h
@@ -0,0 +1,30 @@
+// Copyright Epic Games, Inc. All Rights Reserved.
+
+#pragma once
+
+#include "CoreMinimal.h"
+#include "UTBBaseCommand.h"
+#include "BaseCompositeInlineCommand.generated.h"
+
+/**
+ * 
+ */
+UCLASS(Abstract, EditInlineNew)
+class USERTOOLBOXBASICCOMMAND_API UBaseCompositeInlineCommand : public UUTBBaseCommand
+{
+	GENERATED_BODY()
+public:
+
+	UPROPERTY(EditAnywhere,BlueprintReadWrite,Instanced,Category="CompositeCommand")
+	TArray<TObjectPtr<UUTBBaseCommand>>		Commands;
+	
+	UBaseCompositeInlineCommand()
+	{
+		Name="Toggle Command";
+		Tooltip="This command let you execute sequentially some already registered command, one at a time (placeholder section is a good place to register these commands)";
+		Category="Builtin";
+	}
+
+	
+
+};
diff --git a/Engine/Plugins/Experimental/UserToolBoxBasicCommand/Source/UserToolBoxBasicCommand/Public/ChangeViewMode.h b/Engine/Plugins/Experimental/UserToolBoxBasicCommand/Source/UserToolBoxBasicCommand/Public/ChangeViewMode.h
new file mode 100644
--- /dev/null
+++ b/Engine/Plugins/Experimental/UserToolBoxBasicCommand/Source/UserToolBoxBasicCommand/Public/ChangeViewMode.h
@@ -0,0 +1,22 @@
+// Copyright Epic Games, Inc. All Rights Reserved.
+
+#pragma once
+
+#include "CoreMinimal.h"
+#include "UTBBaseCommand.h"
+#include "ChangeViewMode.generated.h"
+
+/**
+ * 
+ */
+UCLASS()
+class USERTOOLBOXBASICCOMMAND_API UChangeViewMode : public UUTBBaseCommand
+{
+	GENERATED_BODY()
+public:
+
+	UPROPERTY(EditAnywhere,BlueprintReadWrite,Category="ChangeViewCommand")
+	TEnumAsByte<EViewModeIndex> ViewMode=EViewModeIndex::VMI_Lit;
+	virtual void Execute() override;
+	UChangeViewMode();
+};
diff --git a/Engine/Plugins/Experimental/UserToolBoxBasicCommand/Source/UserToolBoxBasicCommand/Public/CleanHierarchy.h b/Engine/Plugins/Experimental/UserToolBoxBasicCommand/Source/UserToolBoxBasicCommand/Public/CleanHierarchy.h
new file mode 100644
--- /dev/null
+++ b/Engine/Plugins/Experimental/UserToolBoxBasicCommand/Source/UserToolBoxBasicCommand/Public/CleanHierarchy.h
@@ -0,0 +1,25 @@
+// Copyright Epic Games, Inc. All Rights Reserved.
+
+#pragma once
+
+#include "CoreMinimal.h"
+#include "UTBBaseCommand.h"
+#include "CleanHierarchy.generated.h"
+
+/**
+ * 
+ */
+UCLASS()
+class USERTOOLBOXBASICCOMMAND_API UCleanHierarchy : public UUTBBaseCommand
+{
+	GENERATED_BODY()
+	public:
+	UPROPERTY(EditAnywhere,Category="CleanHierarchyCommand")
+	TArray<FString> DSMetadataToPreserve;
+	UPROPERTY(EditAnywhere,Category="CleanHierarchyCommand")
+	bool RemoveEmptyBranches;
+	UPROPERTY(EditAnywhere,Category="CleanHierarchyCommand")
+	bool RemoveIntermediaryActorsWithoutGeometry;
+	UCleanHierarchy();
+	virtual void Execute() override;
+};
diff --git a/Engine/Plugins/Experimental/UserToolBoxBasicCommand/Source/UserToolBoxBasicCommand/Public/CompositeCommand.h b/Engine/Plugins/Experimental/UserToolBoxBasicCommand/Source/UserToolBoxBasicCommand/Public/CompositeCommand.h
new file mode 100644
--- /dev/null
+++ b/Engine/Plugins/Experimental/UserToolBoxBasicCommand/Source/UserToolBoxBasicCommand/Public/CompositeCommand.h
@@ -0,0 +1,30 @@
+// Copyright Epic Games, Inc. All Rights Reserved.
+
+#pragma once
+
+#include "CoreMinimal.h"
+#include "BaseCompositeCommand.h"
+#include "UTBBaseCommand.h"
+#include "CompositeCommand.generated.h"
+
+/**
+ * 
+ */
+
+UCLASS(Blueprintable)
+class USERTOOLBOXBASICCOMMAND_API UCompositeCommand : public UBaseCompositeCommand
+{
+	GENERATED_BODY()
+	public:
+
+
+	virtual void Execute() override;
+	UCompositeCommand()
+	{
+		Name="Composite Command";
+		Tooltip="This command let you execute sequentially other already set command (the placeholder section is a good place to register these commands";
+		Category="Utility";
+	}
+
+
+};
diff --git a/Engine/Plugins/Experimental/UserToolBoxBasicCommand/Source/UserToolBoxBasicCommand/Public/CompositeInlineCommand.h b/Engine/Plugins/Experimental/UserToolBoxBasicCommand/Source/UserToolBoxBasicCommand/Public/CompositeInlineCommand.h
new file mode 100644
--- /dev/null
+++ b/Engine/Plugins/Experimental/UserToolBoxBasicCommand/Source/UserToolBoxBasicCommand/Public/CompositeInlineCommand.h
@@ -0,0 +1,30 @@
+// Copyright Epic Games, Inc. All Rights Reserved.
+
+#pragma once
+
+#include "CoreMinimal.h"
+#include "BaseCompositeInlineCommand.h"
+
+#include "CompositeInlineCommand.generated.h"
+
+/**
+ * 
+ */
+
+UCLASS(Blueprintable)
+class USERTOOLBOXBASICCOMMAND_API UCompositeInlineCommand : public UBaseCompositeInlineCommand
+{
+	GENERATED_BODY()
+	public:
+
+
+	virtual void Execute() override;
+	UCompositeInlineCommand()
+	{
+		Name="Composite Inline Command";
+		Tooltip="This command let you execute sequentially other already set command (the placeholder section is a good place to register these commands";
+		Category="Utility";
+	}
+
+
+};
diff --git a/Engine/Plugins/Experimental/UserToolBoxBasicCommand/Source/UserToolBoxBasicCommand/Public/ConsoleVariable.h b/Engine/Plugins/Experimental/UserToolBoxBasicCommand/Source/UserToolBoxBasicCommand/Public/ConsoleVariable.h
new file mode 100644
--- /dev/null
+++ b/Engine/Plugins/Experimental/UserToolBoxBasicCommand/Source/UserToolBoxBasicCommand/Public/ConsoleVariable.h
@@ -0,0 +1,28 @@
+// Copyright Epic Games, Inc. All Rights Reserved.
+
+#pragma once
+
+#include "CoreMinimal.h"
+#include "UTBBaseCommand.h"
+#include "ConsoleVariable.generated.h"
+
+/**
+ * 
+ */
+UCLASS()
+class USERTOOLBOXBASICCOMMAND_API UConsoleVariable : public UUTBBaseCommand
+{
+	GENERATED_BODY()
+	public:
+	UPROPERTY(EditAnywhere,BlueprintReadWrite,Category="Console Command")
+	TArray<FString> ConsoleCommands;
+	
+	virtual void Execute() override;
+
+	UConsoleVariable()
+	{
+		Name="Console Variable Command";
+		Tooltip="This console command let you execute any console variable/command";
+		Category="Utility";
+	}
+};
diff --git a/Engine/Plugins/Experimental/UserToolBoxBasicCommand/Source/UserToolBoxBasicCommand/Public/EngineCommand.h b/Engine/Plugins/Experimental/UserToolBoxBasicCommand/Source/UserToolBoxBasicCommand/Public/EngineCommand.h
new file mode 100644
--- /dev/null
+++ b/Engine/Plugins/Experimental/UserToolBoxBasicCommand/Source/UserToolBoxBasicCommand/Public/EngineCommand.h
@@ -0,0 +1,27 @@
+// Copyright Epic Games, Inc. All Rights Reserved.
+
+#pragma once
+
+#include "CoreMinimal.h"
+#include "UTBBaseCommand.h"
+#include "EngineCommand.generated.h"
+
+/**
+ * 
+ */
+UCLASS()
+class USERTOOLBOXBASICCOMMAND_API UEngineCommand : public UUTBBaseCommand
+{
+	GENERATED_BODY()
+	public:
+
+	UPROPERTY(EditAnywhere,BlueprintReadWrite,Category="EngineCommand")
+	FString Command;
+	virtual void Execute() override;
+	UEngineCommand()
+	{
+		Name="Engine Command";
+		Tooltip="This command let you execute any engine command";
+		Category="Utility";
+	}
+};
diff --git a/Engine/Plugins/Experimental/UserToolBoxBasicCommand/Source/UserToolBoxBasicCommand/Public/ExecuteBindableAction.h b/Engine/Plugins/Experimental/UserToolBoxBasicCommand/Source/UserToolBoxBasicCommand/Public/ExecuteBindableAction.h
new file mode 100644
--- /dev/null
+++ b/Engine/Plugins/Experimental/UserToolBoxBasicCommand/Source/UserToolBoxBasicCommand/Public/ExecuteBindableAction.h
@@ -0,0 +1,37 @@
+// Copyright Epic Games, Inc. All Rights Reserved.
+
+#pragma once
+
+#include "CoreMinimal.h"
+#include "UTBBaseCommand.h"
+
+#include "ExecuteBindableAction.generated.h"
+
+USTRUCT()
+struct USERTOOLBOXBASICCOMMAND_API FBindableActionInfo
+{
+	GENERATED_BODY()
+	UPROPERTY()
+	FString Context;
+	UPROPERTY()
+	FString CommandName;
+};
+
+/**
+ * 
+ */
+UCLASS()
+class USERTOOLBOXBASICCOMMAND_API UExecuteBindableAction : public UUTBBaseCommand
+{
+	GENERATED_BODY()
+public:
+	UExecuteBindableAction() ;
+
+	UPROPERTY(EditAnywhere,Category="Bindable Command")
+	FBindableActionInfo	CommandInfo;
+
+
+public:
+	virtual void Execute() override;
+	
+};
diff --git a/Engine/Plugins/Experimental/UserToolBoxBasicCommand/Source/UserToolBoxBasicCommand/Public/ExecutePythonScript.h b/Engine/Plugins/Experimental/UserToolBoxBasicCommand/Source/UserToolBoxBasicCommand/Public/ExecutePythonScript.h
new file mode 100644
--- /dev/null
+++ b/Engine/Plugins/Experimental/UserToolBoxBasicCommand/Source/UserToolBoxBasicCommand/Public/ExecutePythonScript.h
@@ -0,0 +1,28 @@
+// Copyright Epic Games, Inc. All Rights Reserved.
+
+#pragma once
+
+#include "CoreMinimal.h"
+#include "UTBBaseCommand.h"
+
+#include "ExecutePythonScript.generated.h"
+
+/**
+ * 
+ */
+UCLASS()
+class USERTOOLBOXBASICCOMMAND_API UExecutePythonScript : public UUTBBaseCommand
+{
+	GENERATED_BODY()
+public:
+	UExecutePythonScript() ;
+private:
+	UPROPERTY(EditAnywhere,Category="Python Command")
+	FFilePath	ScriptPath;
+
+	UPROPERTY(EditAnywhere,Category="Python Command")
+	FString		Args;
+
+public:
+	virtual void Execute() override;
+};
diff --git a/Engine/Plugins/Experimental/UserToolBoxBasicCommand/Source/UserToolBoxBasicCommand/Public/FlipNormals.h b/Engine/Plugins/Experimental/UserToolBoxBasicCommand/Source/UserToolBoxBasicCommand/Public/FlipNormals.h
new file mode 100644
--- /dev/null
+++ b/Engine/Plugins/Experimental/UserToolBoxBasicCommand/Source/UserToolBoxBasicCommand/Public/FlipNormals.h
@@ -0,0 +1,20 @@
+// Copyright Epic Games, Inc. All Rights Reserved.
+
+#pragma once
+
+#include "CoreMinimal.h"
+#include "UTBBaseCommand.h"
+#include "FlipNormals.generated.h"
+
+/**
+ * 
+ */
+UCLASS()
+class USERTOOLBOXBASICCOMMAND_API UFlipNormals : public UUTBBaseCommand
+{
+	GENERATED_BODY()
+	public:
+
+	virtual void Execute() override;
+	UFlipNormals();
+};
diff --git a/Engine/Plugins/Experimental/UserToolBoxBasicCommand/Source/UserToolBoxBasicCommand/Public/IsolateSelection.h b/Engine/Plugins/Experimental/UserToolBoxBasicCommand/Source/UserToolBoxBasicCommand/Public/IsolateSelection.h
new file mode 100644
--- /dev/null
+++ b/Engine/Plugins/Experimental/UserToolBoxBasicCommand/Source/UserToolBoxBasicCommand/Public/IsolateSelection.h
@@ -0,0 +1,21 @@
+// Copyright Epic Games, Inc. All Rights Reserved.
+
+#pragma once
+
+#include "CoreMinimal.h"
+#include "UTBBaseCommand.h"
+#include "IsolateSelection.generated.h"
+
+/**
+ * 
+ */
+UCLASS()
+class USERTOOLBOXBASICCOMMAND_API UIsolateSelection : public UUTBBaseCommand
+{
+	GENERATED_BODY()
+	public:
+	UPROPERTY(EditAnywhere,Category="Isolate Command")
+	bool bShouldOnlyAffectStaticMeshActor=false;
+	UIsolateSelection();
+	virtual void Execute() override;
+};
diff --git a/Engine/Plugins/Experimental/UserToolBoxBasicCommand/Source/UserToolBoxBasicCommand/Public/MergeCommand.h b/Engine/Plugins/Experimental/UserToolBoxBasicCommand/Source/UserToolBoxBasicCommand/Public/MergeCommand.h
new file mode 100644
--- /dev/null
+++ b/Engine/Plugins/Experimental/UserToolBoxBasicCommand/Source/UserToolBoxBasicCommand/Public/MergeCommand.h
@@ -0,0 +1,22 @@
+// Copyright Epic Games, Inc. All Rights Reserved.
+
+#pragma once
+
+#include "CoreMinimal.h"
+#include "UTBBaseCommand.h"
+#include "MergeCommand.generated.h"
+
+/**
+ * 
+ */
+UCLASS()
+class USERTOOLBOXBASICCOMMAND_API UMerge : public UUTBBaseCommand
+{
+	GENERATED_BODY()
+	public:
+
+	UPROPERTY(EditAnywhere,BlueprintReadWrite,Category="Merge Command")
+	bool Advanced;
+	UMerge();
+	virtual void Execute() override;
+};
diff --git a/Engine/Plugins/Experimental/UserToolBoxBasicCommand/Source/UserToolBoxBasicCommand/Public/MirrorActorCommand.h b/Engine/Plugins/Experimental/UserToolBoxBasicCommand/Source/UserToolBoxBasicCommand/Public/MirrorActorCommand.h
new file mode 100644
--- /dev/null
+++ b/Engine/Plugins/Experimental/UserToolBoxBasicCommand/Source/UserToolBoxBasicCommand/Public/MirrorActorCommand.h
@@ -0,0 +1,25 @@
+// Copyright Epic Games, Inc. All Rights Reserved.
+
+#pragma once
+
+#include "CoreMinimal.h"
+#include "UTBBaseCommand.h"
+#include "MirrorActorCommand.generated.h"
+
+/**
+ * 
+ */
+UCLASS()
+class USERTOOLBOXBASICCOMMAND_API UMirrorActorCommand : public UUTBBaseCommand
+{
+	GENERATED_BODY()
+	public:
+	UPROPERTY(EditAnywhere,BlueprintReadWrite,Category="MirrorCommand")
+	bool XAxis=false;
+	UPROPERTY(EditAnywhere,BlueprintReadWrite,Category="MirrorCommand")
+	bool YAxis=false;
+	UPROPERTY(EditAnywhere,BlueprintReadWrite,Category="MirrorCommand")
+	bool ZAxis=false;
+	UMirrorActorCommand();
+	virtual void Execute() override;
+};
diff --git a/Engine/Plugins/Experimental/UserToolBoxBasicCommand/Source/UserToolBoxBasicCommand/Public/PushComponentMaterialIntoMesh.h b/Engine/Plugins/Experimental/UserToolBoxBasicCommand/Source/UserToolBoxBasicCommand/Public/PushComponentMaterialIntoMesh.h
new file mode 100644
--- /dev/null
+++ b/Engine/Plugins/Experimental/UserToolBoxBasicCommand/Source/UserToolBoxBasicCommand/Public/PushComponentMaterialIntoMesh.h
@@ -0,0 +1,21 @@
+// Copyright Epic Games, Inc. All Rights Reserved.
+
+#pragma once
+
+#include "CoreMinimal.h"
+#include "UTBBaseCommand.h"
+#include "PushComponentMaterialIntoMesh.generated.h"
+
+/**
+ * 
+ */
+UCLASS()
+class USERTOOLBOXBASICCOMMAND_API UPushComponentMaterialIntoMesh : public UUTBBaseCommand
+{
+	GENERATED_BODY()
+	public:
+
+	void Execute() override;
+	
+	UPushComponentMaterialIntoMesh();
+};
diff --git a/Engine/Plugins/Experimental/UserToolBoxBasicCommand/Source/UserToolBoxBasicCommand/Public/SelectActorByFilter.h b/Engine/Plugins/Experimental/UserToolBoxBasicCommand/Source/UserToolBoxBasicCommand/Public/SelectActorByFilter.h
new file mode 100644
--- /dev/null
+++ b/Engine/Plugins/Experimental/UserToolBoxBasicCommand/Source/UserToolBoxBasicCommand/Public/SelectActorByFilter.h
@@ -0,0 +1,56 @@
+// Copyright Epic Games, Inc. All Rights Reserved.
+
+#pragma once
+
+#include "CoreMinimal.h"
+#include "UTBBaseCommand.h"
+#include "ActorFilter/BaseActorFilter.h"
+#include "SelectActorByFilter.generated.h"
+
+
+UENUM()
+enum class EActorFilterRule : uint8
+{
+	Add =0,
+	Intersect =2,
+	Substract =3,
+	Replace= 4 
+};
+UENUM()
+enum class EActorFilterSource : uint8
+{
+	OriginalSelection =0,
+	PreviousResult =1
+};
+
+
+USTRUCT(BlueprintType)
+struct FActorFilterOptions
+{
+	GENERATED_BODY()
+	UPROPERTY(EditAnywhere, Category="Actor Filter",Instanced)
+	TObjectPtr<UBaseActorFilter> Filter;
+	UPROPERTY(EditAnywhere, Category="Actor Filter")
+	EActorFilterRule Rule=EActorFilterRule::Replace;
+	UPROPERTY(EditAnywhere, Category="Actor Filter")
+	EActorFilterSource Source=EActorFilterSource::PreviousResult;
+	
+	UPROPERTY(EditAnywhere, Category="Actor Filter")
+	bool ShowProperties=false;
+};
+
+
+UCLASS(Blueprintable)
+class USERTOOLBOXBASICCOMMAND_API USelectActorByFilter : public UUTBBaseCommand
+{
+	GENERATED_BODY()
+	public:
+	UPROPERTY(EditAnywhere,BlueprintReadOnly, Category="Actor Filter")
+	bool	ApplyToCurrentSelection;
+	UPROPERTY(EditAnywhere,BlueprintReadWrite, Category="Actor Filter")
+	TArray<FActorFilterOptions>	FilterStack;
+	virtual void Execute() override;
+	USelectActorByFilter();
+	
+	
+};
\ No newline at end of file
diff --git a/Engine/Plugins/Experimental/UserToolBoxBasicCommand/Source/UserToolBoxBasicCommand/Public/SelectActorBySize.h b/Engine/Plugins/Experimental/UserToolBoxBasicCommand/Source/UserToolBoxBasicCommand/Public/SelectActorBySize.h
new file mode 100644
--- /dev/null
+++ b/Engine/Plugins/Experimental/UserToolBoxBasicCommand/Source/UserToolBoxBasicCommand/Public/SelectActorBySize.h
@@ -0,0 +1,22 @@
+// Copyright Epic Games, Inc. All Rights Reserved.
+
+#pragma once
+
+#include "CoreMinimal.h"
+#include "UTBBaseCommand.h"
+#include "SelectActorBySize.generated.h"
+
+/**
+ * 
+ */
+UCLASS()
+class USERTOOLBOXBASICCOMMAND_API USelectActorBySize : public UUTBBaseCommand
+{
+	GENERATED_BODY()
+	public:
+
+	UPROPERTY(EditAnywhere,BlueprintReadWrite,Category="SelectActorBySizeCommand")
+	float SizeThreshold;
+	USelectActorBySize();
+	virtual void Execute() override;
+};
diff --git a/Engine/Plugins/Experimental/UserToolBoxBasicCommand/Source/UserToolBoxBasicCommand/Public/SetHighPrecisionOnMesh.h b/Engine/Plugins/Experimental/UserToolBoxBasicCommand/Source/UserToolBoxBasicCommand/Public/SetHighPrecisionOnMesh.h
new file mode 100644
--- /dev/null
+++ b/Engine/Plugins/Experimental/UserToolBoxBasicCommand/Source/UserToolBoxBasicCommand/Public/SetHighPrecisionOnMesh.h
@@ -0,0 +1,24 @@
+// Copyright Epic Games, Inc. All Rights Reserved.
+
+#pragma once
+
+#include "CoreMinimal.h"
+#include "UTBBaseCommand.h"
+#include "SetHighPrecisionOnMesh.generated.h"
+
+/**
+ * 
+ */
+UCLASS()
+class USERTOOLBOXBASICCOMMAND_API USetHighPrecisionOnMesh : public UUTBBaseCommand
+{
+	GENERATED_BODY()
+
+public:
+	UPROPERTY(EditAnywhere, BlueprintReadWrite,Category="HighPrecisionCommand")
+	bool bHighPrecisionTangent=true;
+	UPROPERTY(EditAnywhere, BlueprintReadWrite,Category="HighPrecisionCommand")
+	bool bHighPrecisionUV=true;
+	USetHighPrecisionOnMesh();
+	virtual void Execute() override;
+};
diff --git a/Engine/Plugins/Experimental/UserToolBoxBasicCommand/Source/UserToolBoxBasicCommand/Public/ShowLayersCommand.h b/Engine/Plugins/Experimental/UserToolBoxBasicCommand/Source/UserToolBoxBasicCommand/Public/ShowLayersCommand.h
new file mode 100644
--- /dev/null
+++ b/Engine/Plugins/Experimental/UserToolBoxBasicCommand/Source/UserToolBoxBasicCommand/Public/ShowLayersCommand.h
@@ -0,0 +1,25 @@
+// Copyright Epic Games, Inc. All Rights Reserved.
+
+#pragma once
+
+#include "CoreMinimal.h"
+#include "UTBBaseCommand.h"
+#include "ShowLayersCommand.generated.h"
+
+/**
+ * 
+ */
+UCLASS()
+class USERTOOLBOXBASICCOMMAND_API UShowLayersCommand : public UUTBBaseCommand
+{
+	GENERATED_BODY()
+public:
+	UPROPERTY(EditAnywhere,BlueprintReadWrite,Category="LayerCommand")
+	TArray<FName>	Layers;
+
+	UPROPERTY(EditAnywhere,BlueprintReadWrite,Category="LayerCommand")
+	bool bIsolate=false;
+	UShowLayersCommand();
+	virtual bool DisplayParameters() override;
+	virtual void Execute() override;
+};
diff --git a/Engine/Plugins/Experimental/UserToolBoxBasicCommand/Source/UserToolBoxBasicCommand/Public/TabSpawner.h b/Engine/Plugins/Experimental/UserToolBoxBasicCommand/Source/UserToolBoxBasicCommand/Public/TabSpawner.h
new file mode 100644
--- /dev/null
+++ b/Engine/Plugins/Experimental/UserToolBoxBasicCommand/Source/UserToolBoxBasicCommand/Public/TabSpawner.h
@@ -0,0 +1,21 @@
+// Copyright Epic Games, Inc. All Rights Reserved.
+
+#pragma once
+
+#include "CoreMinimal.h"
+#include "UTBBaseCommand.h"
+#include "TabSpawner.generated.h"
+
+/**
+ * 
+ */
+UCLASS()
+class USERTOOLBOXBASICCOMMAND_API UTabSpawner : public UUTBBaseCommand
+{
+	GENERATED_BODY()
+	UTabSpawner();
+	UPROPERTY(EditAnywhere,Category="TabSpawnerCommand")
+	FName TabName;
+public:
+	virtual void Execute() override;
+};
diff --git a/Engine/Plugins/Experimental/UserToolBoxBasicCommand/Source/UserToolBoxBasicCommand/Public/ToggleCommand.h b/Engine/Plugins/Experimental/UserToolBoxBasicCommand/Source/UserToolBoxBasicCommand/Public/ToggleCommand.h
new file mode 100644
--- /dev/null
+++ b/Engine/Plugins/Experimental/UserToolBoxBasicCommand/Source/UserToolBoxBasicCommand/Public/ToggleCommand.h
@@ -0,0 +1,34 @@
+// Copyright Epic Games, Inc. All Rights Reserved.
+
+#pragma once
+
+#include "CoreMinimal.h"
+#include "BaseCompositeCommand.h"
+#include "UTBBaseCommand.h"
+#include "CompositeCommand.h"
+#include "ToggleCommand.generated.h"
+
+/**
+ * 
+ */
+
+
+
+UCLASS(Blueprintable)
+class USERTOOLBOXBASICCOMMAND_API UToggleCommand : public UBaseCompositeCommand
+{
+	GENERATED_BODY()
+	public:
+	UToggleCommand()
+	{
+		Name="Toggle Command";
+		Tooltip="This command let you execute sequentially some already registered command, one at a time (placeholder section is a good place to register these commands)";
+		Category="Utility";
+	}
+
+
+private:
+	int CurrentIndex=0;
+	virtual void Execute() override;
+
+};
diff --git a/Engine/Plugins/Experimental/UserToolBoxBasicCommand/Source/UserToolBoxBasicCommand/Public/ToggleCommandInline.h b/Engine/Plugins/Experimental/UserToolBoxBasicCommand/Source/UserToolBoxBasicCommand/Public/ToggleCommandInline.h
new file mode 100644
--- /dev/null
+++ b/Engine/Plugins/Experimental/UserToolBoxBasicCommand/Source/UserToolBoxBasicCommand/Public/ToggleCommandInline.h
@@ -0,0 +1,37 @@
+// Copyright Epic Games, Inc. All Rights Reserved.
+
+#pragma once
+
+#include "CoreMinimal.h"
+#include "UTBBaseCommand.h"
+#include "CompositeCommand.h"
+#include "ToggleCommandInline.generated.h"
+
+/**
+ * 
+ */
+
+
+
+UCLASS(Blueprintable)
+class USERTOOLBOXBASICCOMMAND_API UToggleCommandInline : public UUTBBaseCommand
+{
+	GENERATED_BODY()
+	public:
+
+	UPROPERTY(EditAnywhere,BlueprintReadWrite,Instanced,Category="ToggleCommand")
+	TArray<TObjectPtr<UUTBBaseCommand>>		Commands;
+	
+	UToggleCommandInline()
+	{
+		Name="Toggle Command Inline";
+		Tooltip="This command let you execute sequentially some already registered command, one at a time (placeholder section is a good place to register these commands)";
+		Category="Utility";
+	}
+
+
+private:
+	int CurrentIndex=0;
+	virtual void Execute() override;
+
+};
diff --git a/Engine/Plugins/Experimental/UserToolBoxBasicCommand/Source/UserToolBoxBasicCommand/Public/UserToolBoxBasicCommand.h b/Engine/Plugins/Experimental/UserToolBoxBasicCommand/Source/UserToolBoxBasicCommand/Public/UserToolBoxBasicCommand.h
new file mode 100644
--- /dev/null
+++ b/Engine/Plugins/Experimental/UserToolBoxBasicCommand/Source/UserToolBoxBasicCommand/Public/UserToolBoxBasicCommand.h
@@ -0,0 +1,16 @@
+// Copyright Epic Games, Inc. All Rights Reserved.
+
+#pragma once
+
+#include "CoreMinimal.h"
+#include "Modules/ModuleManager.h"
+ DECLARE_LOG_CATEGORY_EXTERN(LogUserToolBoxBasicCommand, Log, All);
+
+class FUserToolBox_BasicCommandModule : public IModuleInterface
+{
+public:
+
+	/** IModuleInterface implementation */
+	virtual void StartupModule() override;
+	virtual void ShutdownModule() override;
+};
diff --git a/Engine/Plugins/Experimental/UserToolBoxBasicCommand/Source/UserToolBoxBasicCommand/Public/ZoomAll.h b/Engine/Plugins/Experimental/UserToolBoxBasicCommand/Source/UserToolBoxBasicCommand/Public/ZoomAll.h
new file mode 100644
--- /dev/null
+++ b/Engine/Plugins/Experimental/UserToolBoxBasicCommand/Source/UserToolBoxBasicCommand/Public/ZoomAll.h
@@ -0,0 +1,20 @@
+// Copyright Epic Games, Inc. All Rights Reserved.
+
+#pragma once
+
+#include "CoreMinimal.h"
+#include "UTBBaseCommand.h"
+#include "ZoomAll.generated.h"
+
+/**
+ * 
+ */
+UCLASS()
+class USERTOOLBOXBASICCOMMAND_API UZoomAll : public UUTBBaseCommand
+{
+	GENERATED_BODY()
+	public:
+	UZoomAll();
+	virtual void Execute() override;
+	
+};
diff --git a/Engine/Plugins/Experimental/UserToolBoxBasicCommand/Source/UserToolBoxBasicCommand/UserToolBoxBasicCommand.Build.cs b/Engine/Plugins/Experimental/UserToolBoxBasicCommand/Source/UserToolBoxBasicCommand/UserToolBoxBasicCommand.Build.cs
new file mode 100644
--- /dev/null
+++ b/Engine/Plugins/Experimental/UserToolBoxBasicCommand/Source/UserToolBoxBasicCommand/UserToolBoxBasicCommand.Build.cs
@@ -0,0 +1,72 @@
+// Copyright Epic Games, Inc. All Rights Reserved.
+
+using UnrealBuildTool;
+
+public class UserToolBoxBasicCommand : ModuleRules
+{
+	public UserToolBoxBasicCommand(ReadOnlyTargetRules Target) : base(Target)
+	{
+		PCHUsage = ModuleRules.PCHUsageMode.UseExplicitOrSharedPCHs;
+		
+		PublicIncludePaths.AddRange(
+			new string[] {
+				// ... add public include paths required here ...
+			}
+			);
+				
+		
+		PrivateIncludePaths.AddRange(
+			new string[] {
+				// ... add other private include paths required here ...
+			}
+			);
+			
+		
+		PublicDependencyModuleNames.AddRange(
+			new string[]
+			{
+				"Core", "AssetTools"
+				// ... add other public dependencies that you statically link with here ...
+			}
+			);
+			
+		
+		PrivateDependencyModuleNames.AddRange(
+			new string[]
+			{
+				"CoreUObject",
+				"Engine",
+				"Slate",
+				"SlateCore",
+				"UserToolBoxCore",
+				"EditorScriptingUtilities",
+				"UnrealEd",
+				"EditorStyle",
+				 "Blutility", "DatasmithContent",
+				 "EditorInteractiveToolsFramework",
+				 "DatasmithContent",
+				 "EditorWidgets",
+				 "ToolWidgets",
+				 "InputCore",
+				 "PropertyEditor"
+
+				 // ... add private dependencies that you statically link with here ...	
+			}
+			);
+		
+		if (Target.Version.MajorVersion >= 5)
+		{
+			PrivateDependencyModuleNames.AddRange(
+				new string[]
+				{
+					"StaticMeshEditor","GeometryFramework","GeometryScriptingCore"
+				});
+		}
+		DynamicallyLoadedModuleNames.AddRange(
+			new string[]
+			{
+				// ... add any modules that your module loads dynamically here ...
+			}
+			);
+	}
+}
diff --git a/Engine/Plugins/Experimental/UserToolBoxBasicCommand/UserToolBoxBasicCommand.uplugin b/Engine/Plugins/Experimental/UserToolBoxBasicCommand/UserToolBoxBasicCommand.uplugin
new file mode 100644
--- /dev/null
+++ b/Engine/Plugins/Experimental/UserToolBoxBasicCommand/UserToolBoxBasicCommand.uplugin
@@ -0,0 +1,44 @@
+{
+	"FileVersion": 1,
+	"Version": 1,
+	"VersionName": "0.10",
+	"FriendlyName": "UserToolBoxBasicCommand",
+	"Description": "Basic set of command to populate a custom editor tab",
+	"Category": "Editor",
+	"CreatedBy": "Epic Games, Inc.",
+	"CreatedByURL": "https://epicgames.com",
+	"DocsURL": "",
+	"MarketplaceURL": "",
+	"SupportURL": "",
+	"CanContainContent": false,
+	"IsBetaVersion": false,
+	"IsExperimentalVersion": true,
+	"Installed": false,
+	"Modules": [
+		{
+			"Name": "UserToolBoxBasicCommand",
+			"Type": "Editor",
+			"LoadingPhase": "Default"
+		}
+	],
+	"Plugins":
+	[
+		{
+			"Name": "UserToolBoxCore",
+			"Enabled": true
+		},
+		{
+			"Name": "EditorScriptingUtilities",
+			"Enabled" : true
+		},
+		{
+			"Name": "DatasmithContent",
+			"Enabled" : true
+		}
+		,
+		  {
+			"Name": "GeometryScripting",
+			"Enabled" : true
+		  }
+	]
+}
\ No newline at end of file
diff --git a/Engine/Plugins/Experimental/UserToolBoxCore/Source/UserToolBoxCore/Private/AssetTypeActions_UTBEditorUtilityBlueprint.cpp b/Engine/Plugins/Experimental/UserToolBoxCore/Source/UserToolBoxCore/Private/AssetTypeActions_UTBEditorUtilityBlueprint.cpp
new file mode 100644
--- /dev/null
+++ b/Engine/Plugins/Experimental/UserToolBoxCore/Source/UserToolBoxCore/Private/AssetTypeActions_UTBEditorUtilityBlueprint.cpp
@@ -0,0 +1,40 @@
+// Copyright Epic Games, Inc. All Rights Reserved.
+
+#include "AssetTypeActions_UTBEditorUtilityBlueprint.h"
+#include "GlobalEditorUtilityBase.h"
+#include "IContentBrowserSingleton.h"
+#include "UserToolBoxSubsystem.h"
+#include "UTBEditorBlueprint.h"
+#define LOCTEXT_NAMESPACE "AssetTypeActions"
+
+/////////////////////////////////////////////////////
+// FAssetTypeActions_EditorUtilityBlueprint
+
+FText FAssetTypeActions_UTBEditorBlueprint::GetName() const
+{
+	return FText::FromString("User ToolBox Command Blueprint");
+}
+
+FColor FAssetTypeActions_UTBEditorBlueprint::GetTypeColor() const
+{
+	return FColor(0, 169, 0);
+}
+
+UClass* FAssetTypeActions_UTBEditorBlueprint::GetSupportedClass() const
+{
+	return UUTBEditorBlueprint::StaticClass();
+}
+
+bool FAssetTypeActions_UTBEditorBlueprint::HasActions(const TArray<UObject*>& InObjects) const
+{
+	return false;
+}
+uint32 FAssetTypeActions_UTBEditorBlueprint::GetCategories()
+{
+	return UUserToolboxSubsystem::AssetCategory;
+}
+
+
+
+
+#undef LOCTEXT_NAMESPACE
diff --git a/Engine/Plugins/Experimental/UserToolBoxCore/Source/UserToolBoxCore/Private/AssetTypeActions_UTBEditorUtilityBlueprint.h b/Engine/Plugins/Experimental/UserToolBoxCore/Source/UserToolBoxCore/Private/AssetTypeActions_UTBEditorUtilityBlueprint.h
new file mode 100644
--- /dev/null
+++ b/Engine/Plugins/Experimental/UserToolBoxCore/Source/UserToolBoxCore/Private/AssetTypeActions_UTBEditorUtilityBlueprint.h
@@ -0,0 +1,20 @@
+// Copyright Epic Games, Inc. All Rights Reserved.
+
+#pragma once
+
+#include "CoreMinimal.h"
+#include "AssetTypeActions/AssetTypeActions_Blueprint.h"
+
+
+class FAssetTypeActions_UTBEditorBlueprint : public FAssetTypeActions_Blueprint
+{
+public:
+	// IAssetTypeActions interface
+	virtual FText GetName() const override;
+	virtual FColor GetTypeColor() const override;
+	virtual UClass* GetSupportedClass() const override;
+	virtual bool HasActions(const TArray<UObject*>& InObjects) const override;
+	virtual uint32 GetCategories() override;
+	virtual bool CanLocalize() const override { return false; }
+	// End of IAssetTypeActions interface
+};
diff --git a/Engine/Plugins/Experimental/UserToolBoxCore/Source/UserToolBoxCore/Private/Customization/BaseCommandCustomization.cpp b/Engine/Plugins/Experimental/UserToolBoxCore/Source/UserToolBoxCore/Private/Customization/BaseCommandCustomization.cpp
new file mode 100644
--- /dev/null
+++ b/Engine/Plugins/Experimental/UserToolBoxCore/Source/UserToolBoxCore/Private/Customization/BaseCommandCustomization.cpp
@@ -0,0 +1,82 @@
+﻿// Copyright Epic Games, Inc. All Rights Reserved.
+
+
+#include "BaseCommandCustomization.h"
+
+#include "DetailCategoryBuilder.h"
+#include "DetailLayoutBuilder.h"
+#include "DetailWidgetRow.h"
+#include "UserToolBoxSubsystem.h"
+#include "UTBBaseCommand.h"
+#include "Widgets/Input/SButton.h"
+#include "Editor.h"
+#include "Widgets/Images/SImage.h"
+
+void FBaseCommandCustomization::CustomizeDetails(IDetailLayoutBuilder& DetailBuilder)
+{
+	TSharedRef<IPropertyHandle> PropertyHandle= DetailBuilder.GetProperty(GET_MEMBER_NAME_CHECKED(UUTBBaseCommand,IconPath));
+	TArray<UObject*> Outers;
+	PropertyHandle->GetOuterObjects(Outers);
+	IDetailPropertyRow* PropertyRow=DetailBuilder.EditDefaultProperty(PropertyHandle);
+	if (PropertyRow==nullptr)
+	{
+		return;
+	}
+		
+	DetailBuilder.EditDefaultProperty(PropertyHandle)->CustomWidget()
+	.ValueWidget
+	[
+		SNew(SHorizontalBox)
+			+SHorizontalBox::Slot()
+			.HAlign(HAlign_Left)
+			.AutoWidth()
+			[
+				PropertyHandle->CreatePropertyValueWidget()
+			]
+			+SHorizontalBox::Slot()
+			.HAlign(HAlign_Left)
+			.AutoWidth()
+			[
+				
+				SNew(SButton)
+				.ButtonStyle(&FAppStyle::Get().GetWidgetStyle<FButtonStyle>( "Button" ))
+				.ContentPadding(1.0f)
+				.OnClicked_Lambda([Outers]()
+				{
+					
+					if (FString IconValue=""; GEditor->GetEditorSubsystem<UUserToolboxSubsystem>()->PickAnIcon(IconValue))
+					{
+						for (UObject* Object:Outers)
+						{
+							
+							if (UUTBBaseCommand* Command=Cast<UUTBBaseCommand>(Object);IsValid(Command))
+							{
+								Command->IconPath=IconValue;
+								Command->PostEditChange();
+							}
+						}
+						
+						
+					}
+					return FReply::Handled();
+				})
+				[
+					SNew(SImage)
+					.ColorAndOpacity(FSlateColor::UseForeground())
+					.Image(FAppStyle::Get().GetBrush("ShowFlagsMenu.SubMenu.Sprites"))
+				]
+			]
+	]
+	.NameWidget
+	[
+		PropertyHandle->CreatePropertyNameWidget()
+	]
+	;
+	
+	
+}
+
+void FBaseCommandCustomization::CustomizeDetails(const TSharedPtr<IDetailLayoutBuilder>& DetailBuilder)
+{
+	CustomizeDetails(*DetailBuilder);
+}
diff --git a/Engine/Plugins/Experimental/UserToolBoxCore/Source/UserToolBoxCore/Private/Customization/BaseCommandCustomization.h b/Engine/Plugins/Experimental/UserToolBoxCore/Source/UserToolBoxCore/Private/Customization/BaseCommandCustomization.h
new file mode 100644
--- /dev/null
+++ b/Engine/Plugins/Experimental/UserToolBoxCore/Source/UserToolBoxCore/Private/Customization/BaseCommandCustomization.h
@@ -0,0 +1,32 @@
+﻿// Copyright Epic Games, Inc. All Rights Reserved.
+
+#pragma once
+
+#include "CoreMinimal.h"
+#include "IDetailCustomization.h"
+
+/**
+ * 
+ */
+
+
+
+
+class FBaseCommandCustomization : public IDetailCustomization
+{
+public:
+	static TSharedRef<IDetailCustomization> MakeInstance()
+	{
+		return MakeShareable(new FBaseCommandCustomization());
+	}
+	virtual ~FBaseCommandCustomization() override;
+
+	virtual void CustomizeDetails(IDetailLayoutBuilder& DetailBuilder) override;
+	virtual void CustomizeDetails(const TSharedPtr<IDetailLayoutBuilder>& DetailBuilder) override;
+};
+
+inline FBaseCommandCustomization::~FBaseCommandCustomization()
+{
+}
+
+
diff --git a/Engine/Plugins/Experimental/UserToolBoxCore/Source/UserToolBoxCore/Private/Editor/SCommandPickMenu.cpp b/Engine/Plugins/Experimental/UserToolBoxCore/Source/UserToolBoxCore/Private/Editor/SCommandPickMenu.cpp
new file mode 100644
--- /dev/null
+++ b/Engine/Plugins/Experimental/UserToolBoxCore/Source/UserToolBoxCore/Private/Editor/SCommandPickMenu.cpp
@@ -0,0 +1,290 @@
+﻿// Copyright Epic Games, Inc. All Rights Reserved.
+
+
+#include "Editor/SCommandPickMenu.h"
+
+#include "UserToolBoxBaseBlueprint.h"
+#include "UTBEditorBlueprint.h"
+#include "AssetRegistry/AssetRegistryModule.h"
+#include "Editor/UTBTabEditor.h"
+
+#include "Widgets/Input/SSearchBox.h"
+#include "Widgets/Layout/SScrollBorder.h"
+#include "Widgets/Layout/SScrollBox.h"
+#include "Widgets/Layout/SWrapBox.h"
+#include "Widgets/Views/STreeView.h"
+#include "Widgets/Text/STextBlock.h"
+#include "UObject/UObjectIterator.h"
+
+SCommandPickMenuWidget::SCommandPickMenuWidget()
+{
+}
+
+SCommandPickMenuWidget::~SCommandPickMenuWidget()
+{
+}
+
+void SCommandPickMenuWidget::Construct(const FArguments& InArgs)
+{
+	OnCommandSelectionChanged=InArgs._OnCommandSelectionChanged;
+	ReconstructCommandTree();
+	
+	
+	SAssignNew(TreeViewPtr,STreeView<TSharedPtr<FCommandTreeNode>>)
+		.AllowOverscroll(EAllowOverscroll::Yes)
+		.ItemHeight(24)
+		.TreeItemsSource(&CommandTree)
+		.OnSelectionChanged_Lambda([this](TSharedPtr<FCommandTreeNode> Node, ESelectInfo::Type)
+		{
+			if (!Node.IsValid())
+			{
+				Selection=nullptr;
+				OnCommandSelectionChanged.ExecuteIfBound(nullptr);
+				return;
+			}
+			if (!Node->bIsCategory)
+			{
+				Selection=Node->Class;
+				
+				OnCommandSelectionChanged.ExecuteIfBound(Node->Class);
+				
+				
+			}
+			else
+			{
+				Selection=nullptr;
+				OnCommandSelectionChanged.ExecuteIfBound(nullptr);
+			}	
+		})
+		.OnContextMenuOpening(FOnContextMenuOpening::CreateLambda([this]()
+		{
+			if (Selection!=nullptr)
+			{
+				return SNew(STextBlock).Text(FText::FromString(Selection->GetFullName()));
+			}
+			return SNew(STextBlock).Text(FText::FromString("Nullptr"));
+			
+		}))
+		.OnSetExpansionRecursive(this, &SCommandPickMenuWidget::OnSetExpansionRecursive)
+		.OnGenerateRow(STreeView<TSharedPtr<FCommandTreeNode>>::FOnGenerateRow::CreateLambda([](TSharedPtr<FCommandTreeNode> Node,const TSharedRef<STableViewBase>& OwnerTable)
+		{
+			TSharedPtr<STextBlock> TextBlock;
+			if (Node.IsValid())
+			{
+				if (Node->bIsCategory)
+				{
+					SAssignNew(TextBlock,STextBlock).Text(FText::FromString(Node->Name))
+					.Font( FCoreStyle::GetDefaultFontStyle("Bold", 9) );	
+				}
+				else
+				{
+					SAssignNew(TextBlock,STextBlock).Text(FText::FromString(Node->Name))
+					.Font(FCoreStyle::GetDefaultFontStyle("Regular", 9))
+					.ToolTipText_Lambda([Node]()
+					{
+						UUTBBaseCommand* Command=Cast<UUTBBaseCommand>(Node->Class->GetDefaultObject());
+						if (IsValid(Command))
+						{
+							return FText::FromString(Command->Tooltip);	
+						}else
+						{
+							return FText::FromString("Tooltip not available");
+						}
+												
+					})
+					
+					;
+					
+				}
+				
+
+			}
+				
+			else
+			{
+				SAssignNew(TextBlock,STextBlock).Text(FText::FromString("Invalid node"));
+			}
+			
+			
+			return SNew(STableRow<TSharedPtr<FCommandTreeNode>>,OwnerTable).OnDragDetected(FOnDragDetected::CreateLambda([Node](const FGeometry& InGeometry, const FPointerEvent& InPointerEvent )
+			{
+				if (Node->bIsCategory)
+				{
+					return FReply::Unhandled();
+				}
+				return FReply::Handled().BeginDragDrop(FCommandDragDropOperation::New(Cast<UUTBBaseCommand>(Node->Class->GetDefaultObject())));	
+			}))
+			[
+				TextBlock.ToSharedRef()
+			];
+		}))
+		.OnGetChildren(STreeView<TSharedPtr<FCommandTreeNode>>::FOnGetChildren::CreateLambda([this](TSharedPtr<FCommandTreeNode> ParentNode,TArray<TSharedPtr<FCommandTreeNode>>& Children)
+		{
+			
+			
+			Children=ParentNode->FilteredChildren;
+
+		}));
+		
+
+	this->ChildSlot
+	[
+		SNew(SVerticalBox)
+		
+		+SVerticalBox::Slot()
+		.AutoHeight()
+		.VAlign(EVerticalAlignment::VAlign_Top)
+		[
+			SNew(SSearchBox)
+			.OnTextChanged(FOnTextChanged::CreateLambda([this](const FText& Text)
+			{
+				this->UpdateFilteredTreeNodes(Text.ToString());
+			}))
+		]
+		+SVerticalBox::Slot()
+		[
+			TreeViewPtr.ToSharedRef()
+		]
+		
+	];
+	UpdateFilteredTreeNodes("");
+}
+
+void SCommandPickMenuWidget::ReconstructCommandTree()
+{
+	CommandTree.Empty();
+	CommandTreeAllNodes.Empty();
+	LeafsWithParentHierarchy.Empty();
+
+	/*
+	 * Need o load blueprint command before
+	 */
+	
+	IAssetRegistry& AssetRegistry = FModuleManager::LoadModuleChecked<FAssetRegistryModule>(TEXT("AssetRegistry")).Get();
+	FARFilter Filter;
+	Filter.bIncludeOnlyOnDiskAssets = false;
+	Filter.bRecursivePaths = true;
+	Filter.bRecursiveClasses = true;
+	Filter.ClassPaths = {UUTBEditorBlueprint::StaticClass()->GetClassPathName()};
+	TArray<FAssetData> CommandBlueprintAssetDatas;
+	AssetRegistry.GetAssets(Filter,CommandBlueprintAssetDatas);
+	for (FAssetData AssetData:CommandBlueprintAssetDatas)
+	{
+		AssetData.GetAsset();
+	}
+	TArray<UClass*> Classes;
+	for( TObjectIterator<UClass> It ; It ; ++It )
+	{
+		UClass* CurrentClass=*It;
+		if (CurrentClass->IsChildOf(UUTBBaseCommand::StaticClass()))
+		{
+			if (CurrentClass->GetName().StartsWith(TEXT("SKEL_")) || CurrentClass->GetName().StartsWith(TEXT("REINST_")))
+			{
+				continue;
+			}
+			if ( (CurrentClass)!= UUTBBaseCommand::StaticClass() && (CurrentClass)!=UUTBEditorBlueprint::StaticClass() && (CurrentClass)!=UUserToolBoxBaseBlueprint::StaticClass() && !CurrentClass->HasAnyClassFlags(EClassFlags::CLASS_Abstract))
+			{
+				Classes.Add(*It);
+			}
+			
+		}
+	}
+
+	TMap<FString,TArray<UClass*>> UClassSortedByCategory;
+
+	for(UClass* Class:Classes)
+	{
+		UClassSortedByCategory.FindOrAdd(Cast<UUTBBaseCommand>(Class->GetDefaultObject())->Category).Add(Class);
+	}
+
+
+	for (TTuple<FString,TArray<UClass*>> Category:UClassSortedByCategory)
+	{
+		FString CategoryFullName=Category.Key;
+		if (CategoryFullName.IsEmpty())
+		{
+			CategoryFullName="Uncategorized";
+		}
+		TArray<FString> CategoriesSplit;
+		CategoryFullName.ParseIntoArray(CategoriesSplit,TEXT("|"),1);
+		TArray<TSharedPtr<FCommandTreeNode>>* CurrentArray=&CommandTree;
+		TArray<TSharedPtr<FCommandTreeNode>> HierarchyHistory;
+		while (!CategoriesSplit.IsEmpty())
+		{
+			TSharedPtr<FCommandTreeNode>* CurrentNode= CurrentArray->FindByPredicate([&CategoriesSplit](TSharedPtr<FCommandTreeNode> Node)
+			{
+				return CategoriesSplit[0].Compare(Node->Name)==0;
+			});
+			if (CurrentNode==nullptr)
+			{
+				TSharedPtr<FCommandTreeNode> NewNode=CurrentArray->Add_GetRef(MakeShared<FCommandTreeNode>());
+				NewNode->bIsCategory=true;
+				NewNode->Name=CategoriesSplit[0];
+				CurrentNode=&NewNode;
+				HierarchyHistory.Add(NewNode);
+				CommandTreeAllNodes.Add(NewNode);
+			}
+			HierarchyHistory.Add(*CurrentNode);
+			CurrentArray=&(*CurrentNode)->Children;
+			CategoriesSplit.RemoveAt(0);
+		}
+		for (UClass* Class:Category.Value)
+		{
+			TSharedPtr<FCommandTreeNode> NewNode=CurrentArray->Add_GetRef(MakeShared<FCommandTreeNode>());
+			NewNode->Name=Class->GetDisplayNameText().ToString();
+			NewNode->bIsCategory=false;
+			NewNode->Class=Class;
+			LeafsWithParentHierarchy.Add(NewNode,HierarchyHistory);
+			CommandTreeAllNodes.Add(NewNode);
+		}
+	}
+	
+	
+}
+
+void SCommandPickMenuWidget::UpdateFilteredTreeNodes(FString Text)
+{
+
+	
+	TArray<TSharedPtr<FCommandTreeNode>> Leafs;
+	LeafsWithParentHierarchy.GetKeys(Leafs);
+	TArray<TSharedPtr<FCommandTreeNode>> FilteredLeafs;
+	FilteredLeafs=Leafs.FilterByPredicate([Text](const TSharedPtr<FCommandTreeNode>& Object)
+	{
+		return Cast<UUTBBaseCommand>(Object->Class->GetDefaultObject())->Name.Contains(Text) || Text.IsEmpty(); 
+	});
+
+	TSet<TSharedPtr<FCommandTreeNode>> FilteredNodes;
+	for (TSharedPtr<FCommandTreeNode> FilteredLeaf:FilteredLeafs)
+	{
+		FilteredNodes.Append(LeafsWithParentHierarchy.FindChecked(FilteredLeaf));
+	}
+	
+	for (TSharedPtr<FCommandTreeNode> FilteredNode:CommandTreeAllNodes)
+	{
+		FilteredNode->FilteredChildren.Empty();
+		for (TSharedPtr<FCommandTreeNode> Child:FilteredNode->Children)
+		{
+			if (FilteredNodes.Contains(Child) || FilteredLeafs.Contains(Child))
+			{
+				FilteredNode->FilteredChildren.Add(Child);
+			}
+		}
+		TreeViewPtr->SetItemExpansion(FilteredNode,FilteredNode->FilteredChildren.Num()>0);
+	}
+	
+	
+	TreeViewPtr->RequestTreeRefresh();
+}
+void SCommandPickMenuWidget::OnSetExpansionRecursive(TSharedPtr<FCommandTreeNode> InTreeNode, bool bInIsItemExpanded)
+{
+	if (InTreeNode.IsValid() && InTreeNode->Children.Num())
+	{
+		TreeViewPtr->SetItemExpansion(InTreeNode, bInIsItemExpanded);
+
+		for (TSharedPtr<FCommandTreeNode> Child : InTreeNode->Children)
+		{
+			OnSetExpansionRecursive(Child, bInIsItemExpanded);
+		}
+	}
+}
diff --git a/Engine/Plugins/Experimental/UserToolBoxCore/Source/UserToolBoxCore/Private/Editor/SCommandWithOptionsDetailsView.cpp b/Engine/Plugins/Experimental/UserToolBoxCore/Source/UserToolBoxCore/Private/Editor/SCommandWithOptionsDetailsView.cpp
new file mode 100644
--- /dev/null
+++ b/Engine/Plugins/Experimental/UserToolBoxCore/Source/UserToolBoxCore/Private/Editor/SCommandWithOptionsDetailsView.cpp
@@ -0,0 +1,60 @@
+﻿// Copyright Epic Games, Inc. All Rights Reserved.
+
+
+#include "Editor/SCommandWithOptionsDetailsView.h"
+#include "SDropTarget.h"
+#include "UserToolBoxStyle.h"
+#include "UserToolBoxSubsystem.h"
+#include "UTBBaseCommand.h"
+#include "Widgets/Text/SInlineEditableTextBlock.h"
+#include "PropertyEditorModule.h"
+#include "Modules/ModuleManager.h"
+SCommandDetailsView::SCommandDetailsView()
+{
+}
+
+SCommandDetailsView::~SCommandDetailsView()
+{
+}
+
+void SCommandDetailsView::Construct(const FArguments& InArgs)
+{
+	Command=InArgs._Command.Get();
+	OnObjectPropertyModified=InArgs._OnObjectPropertyModified;
+	FPropertyEditorModule& PropertyEditorModule = FModuleManager::GetModuleChecked<FPropertyEditorModule>("PropertyEditor");
+	FDetailsViewArgs Args;
+	Args.bAllowSearch=true;
+	Args.bShowOptions=true;
+	Args.bHideSelectionTip=true;
+	Args.bShowKeyablePropertiesOption=false;
+	Args.bShowObjectLabel=false;
+	Args.bUpdatesFromSelection=false;
+	
+	Args.NameAreaSettings=FDetailsViewArgs::ObjectsUseNameArea;
+	CommandDetailsView=PropertyEditorModule.CreateDetailView(Args);
+	CommandDetailsView->SetIsPropertyVisibleDelegate(FIsPropertyVisible::CreateLambda([](const FPropertyAndParent& PropertyAndParent)
+	{
+		TArray<FName> PropertyNameToHide;
+		PropertyNameToHide.Add(GET_MEMBER_NAME_CHECKED(UUTBBaseCommand,Category));
+		return !PropertyNameToHide.Contains(PropertyAndParent.Property.GetFName());
+	}));
+	this->ChildSlot
+	[
+		SNew(SHorizontalBox)
+		+SHorizontalBox::Slot()
+		.FillWidth(1.0f)
+		.HAlign(HAlign_Fill)
+		.VAlign(VAlign_Fill)
+		[
+			CommandDetailsView.ToSharedRef()
+		]
+	
+	];
+}
+
+void SCommandDetailsView::SetObject(UUTBBaseCommand* InCommand)
+{
+	
+	CommandDetailsView->SetObject(InCommand);
+	Command=InCommand;
+}
diff --git a/Engine/Plugins/Experimental/UserToolBoxCore/Source/UserToolBoxCore/Private/Editor/SCommandWithOptionsWidget.cpp b/Engine/Plugins/Experimental/UserToolBoxCore/Source/UserToolBoxCore/Private/Editor/SCommandWithOptionsWidget.cpp
new file mode 100644
--- /dev/null
+++ b/Engine/Plugins/Experimental/UserToolBoxCore/Source/UserToolBoxCore/Private/Editor/SCommandWithOptionsWidget.cpp
@@ -0,0 +1,355 @@
+﻿// Copyright Epic Games, Inc. All Rights Reserved.
+
+
+#include "Editor/SCommandWithOptionsWidget.h"
+
+#include "IHeadMountedDisplay.h"
+#include "ISequencer.h"
+#include "SDropTarget.h"
+#include "ToolMenuDelegates.h"
+#include "ToolMenus.h"
+#include "UserToolBoxStyle.h"
+#include "UserToolBoxSubsystem.h"
+#include "Editor/UTBTabEditor.h"
+#include "UTBBaseCommand.h"
+#include "AssetRegistry/AssetRegistryModule.h"
+#include "Widgets/Text/SInlineEditableTextBlock.h"
+#include "Widgets/Images/SImage.h"
+#include "Layout/WidgetPath.h"
+#include "Framework/Application/SlateApplication.h"
+
+SButtonCommandWidget::SButtonCommandWidget()
+{
+}
+
+SButtonCommandWidget::~SButtonCommandWidget()
+{
+}
+
+void SButtonCommandWidget::Construct(const FArguments& InArgs)
+{
+	Command=InArgs._Command.Get();
+	SButton::FArguments ButtonArgs;
+	ButtonArgs._OnClicked=InArgs._OnClicked;
+	ButtonArgs._Content=InArgs._Content;
+	SButton::Construct(ButtonArgs);
+	
+}
+
+FReply SButtonCommandWidget::OnDragDetected(const FGeometry& MyGeometry, const FPointerEvent& MouseEvent)
+{
+	
+	return FReply::Handled().BeginDragDrop(FCommandDragDropOperation::New(Command.Get()));
+}
+
+FReply SButtonCommandWidget::OnMouseButtonDown(const FGeometry& MyGeometry, const FPointerEvent& MouseEvent)
+{
+	if (MouseEvent.GetPressedButtons().Contains(EKeys::RightMouseButton))
+	{
+		UE_LOG(LogTemp, Display, TEXT("test"));
+		return FReply::Handled();
+	}
+	return SButton::OnMouseButtonDown(MyGeometry, MouseEvent).DetectDrag(SharedThis(this),EKeys::LeftMouseButton);
+}
+
+
+SCommandWidget::SCommandWidget():IndexInSection(-1)
+{
+}
+
+SCommandWidget::~SCommandWidget()
+{
+	FCoreUObjectDelegates::OnObjectsReplaced.RemoveAll(this);
+}
+
+FReply SCommandWidget::OnDragDetected(const FGeometry& MyGeometry, const FPointerEvent& MouseEvent)
+{
+	return FReply::Handled().BeginDragDrop(FCommandDragDropOperation::New(Command.Get()));
+}
+
+FReply SCommandWidget::OnMouseButtonDown(const FGeometry& MyGeometry, const FPointerEvent& MouseEvent)
+{
+	if (MouseEvent.GetEffectingButton() == EKeys::LeftMouseButton)
+	{
+		return FReply::Handled().DetectDrag(SharedThis(this), EKeys::LeftMouseButton);
+	}
+	return FReply::Unhandled();
+}
+
+void SCommandWidget::Construct(const FArguments& InArgs)
+{
+	Command=InArgs._Command.Get();
+	OnTabChanged=InArgs._OnTabChanged;
+	OnCommandChanged=InArgs._OnCommandChanged;
+	OnCommandSelected=InArgs._OnCommandSelected;
+	OnCommandDeleted=InArgs._OnCommandDeleted;
+	OnDropCommand=InArgs._OnDropCommand;
+	IndexInSection=InArgs._IndexInSection.Get();
+	SectionName=InArgs._SectionName.Get();
+	OnContextMenuOpening=InArgs._OnContextMenuOpening;
+	Editor=InArgs._Editor.Get();
+
+	OnContextMenuOpening.BindRaw(this,&SCommandWidget::CreateContextMenu);
+	//if it's a bp command we should handle recompile
+	if (IsValid(Command.Get()) && Command->GetClass()->ClassGeneratedBy!=nullptr)
+	{
+		FCoreUObjectDelegates::OnObjectsReplaced.AddRaw(this, &SCommandWidget::OnObjectsReplacedCallback);
+
+	}
+	RebuildWidget();
+}
+
+void SCommandWidget::OnObjectsReplacedCallback(const TMap<UObject*, UObject*>& ReplacementMap)
+{
+	UUTBBaseCommand* CommandPtr=Command.Get(true);
+	UObject*const * ReplacementPtrPtr=ReplacementMap.Find(CommandPtr);
+	
+	if (ReplacementPtrPtr!=nullptr)
+	{
+		UUTBBaseCommand* ReplacementPtr=Cast<UUTBBaseCommand>(*ReplacementPtrPtr);
+		if (IsValid(ReplacementPtr))
+		{
+			Command=ReplacementPtr;	
+		}
+	}
+		
+}
+
+void SCommandWidget::RebuildWidget()
+{
+	if (!Command.IsValid())
+	{
+		return;
+	}
+	TSharedPtr<SWidget> IconWidget;
+	if (Command->IconPath.IsEmpty())
+	{
+		IconWidget=SNew(SVerticalBox);
+	}
+	else
+	{
+		IconWidget=SNew(SImage)
+		.Image(FUserToolBoxStyle::Get().GetBrush(FName(Command->IconPath)));
+	}
+	
+	
+	this->ChildSlot
+	[
+		SNew(SBox)
+		.HeightOverride(125)
+		.WidthOverride(250)
+		.Padding(10)
+
+		[
+			SNew(SDropTarget)
+			.OnAllowDrop(SDropTarget::FVerifyDrag::CreateLambda([](TSharedPtr<FDragDropOperation> Operation)
+			{
+
+				return Operation->IsOfType<FCommandDragDropOperation>();
+			}))
+			.OnDropped(FOnDrop::CreateLambda([this](const FGeometry& Geometry,const FDragDropEvent& Event)
+			{
+				if ( Event.GetOperation() && Event.GetOperation()->IsOfType<FCommandDragDropOperation>())
+				{
+					OnDropCommand.ExecuteIfBound(Event.GetOperationAs<FCommandDragDropOperation>()->Command.Get(),SectionName,IndexInSection);
+				}
+			return FReply::Handled().EndDragDrop();	
+			}))
+			[
+				SNew(SButtonCommandWidget)
+				.Command(Command.Get())
+				
+				.OnClicked_Lambda([this]()
+				{
+					
+					OnCommandSelected.ExecuteIfBound(this->Command.Get());
+					return FReply::Handled();
+				})
+				[
+					SNew(SHorizontalBox)
+					+SHorizontalBox::Slot()
+					[
+						SNew(SVerticalBox)
+						+SVerticalBox::Slot()
+						.HAlign(HAlign_Center)
+						.VAlign(VAlign_Fill)
+						[
+							SNullWidget::NullWidget
+						]
+						+SVerticalBox::Slot()
+						.AutoHeight()
+						.HAlign(HAlign_Center)
+						.VAlign(VAlign_Center)
+						[
+							IconWidget.ToSharedRef()
+						]
+						+SVerticalBox::Slot()
+						.AutoHeight()
+						.HAlign(HAlign_Center)
+						.VAlign(VAlign_Center)
+						[
+							SNew(STextBlock)
+							.Text(TAttribute<FText>::Create([this]()
+							{
+								return FText::FromString(Command->Name);
+							}))
+						]
+						+SVerticalBox::Slot()
+						.HAlign(HAlign_Center)
+						.VAlign(VAlign_Fill)
+						[
+							SNullWidget::NullWidget
+						]
+					]
+					+SHorizontalBox::Slot()
+					.HAlign(HAlign_Right)
+					.VAlign(VAlign_Fill)
+					.AutoWidth()
+					[
+						SNew(SVerticalBox)
+						+SVerticalBox::Slot()
+						.AutoHeight()
+						.VAlign(VAlign_Top)
+						[
+								SNew(SButton)
+								.OnClicked_Lambda([this]()
+								{
+									bool ret=this->OnCommandDeleted.ExecuteIfBound(this->Command.Get());
+									return FReply::Handled();
+								})
+								[
+									SNew(SImage)
+									.Image( FAppStyle::GetBrush("Icons.Delete") )
+								]
+						]
+						+SVerticalBox::Slot()
+						.VAlign(VAlign_Fill)
+						[
+							SNullWidget::NullWidget
+						]
+					]
+				]
+			]
+		]
+	];
+}
+static const FName DefaultContextBaseMenuName("UserToolboxTabEditor.DefaultContextMenuBase");
+static const FName DefaultContextMenuName("UserToolboxTabEditor.DefaultContextMenu");
+void SCommandWidget::RegisterContextMenu()
+{
+	UToolMenus* ToolMenus = UToolMenus::Get();
+
+	if (!ToolMenus->IsMenuRegistered(DefaultContextBaseMenuName))
+	{
+		UToolMenu* Menu = ToolMenus->RegisterMenu(DefaultContextBaseMenuName);
+
+		TWeakPtr<FUTBTabEditor> EditorPtr=Editor;
+		Menu->AddDynamicSection("Copy/Move Section", FNewToolMenuDelegate::CreateLambda([EditorPtr](UToolMenu* InMenu)
+		{
+			UUTBBaseCommand* CurrentCommand=InMenu->Context.FindContext<UUTBBaseCommand>();
+			if (!IsValid(CurrentCommand))
+			{
+				return;
+			}
+			IAssetRegistry& AssetRegistry = FModuleManager::LoadModuleChecked<FAssetRegistryModule>(TEXT("AssetRegistry")).Get();
+			FARFilter Filter;
+			Filter.bIncludeOnlyOnDiskAssets = false;
+			Filter.ClassPaths = {UUserToolBoxBaseTab::StaticClass()->GetClassPathName()};
+			Filter.bRecursivePaths = true;
+			Filter.bRecursiveClasses = true;
+			TArray<FAssetData> AssetDatas;
+			AssetRegistry.GetAssets(Filter,AssetDatas);
+			FToolMenuSection& Section = InMenu->AddSection("Copy/Move Section2", NSLOCTEXT("UserToolbox","CommandContextMenu", "Copy/Move"));
+			Section.AddMenuEntry("Duplicate", NSLOCTEXT("UserToolbox","Duplicate Command", "Duplicate"), FText(), FSlateIcon(),
+				FUIAction(FExecuteAction::CreateLambda([EditorPtr,CurrentCommand]()
+				{
+					CurrentCommand->CopyCommand(CurrentCommand->GetOuter());
+					if (EditorPtr.IsValid())
+					{
+						TSharedPtr<FUTBTabEditor> TabEditor=EditorPtr.Pin();
+						TabEditor->NotifyTabHasChanged(Cast<UUserToolBoxBaseTab>(CurrentCommand->GetOuter()->GetOuter()));
+						TabEditor->UpdateCurrentSection();
+					}
+					
+				})));
+			
+			Section.AddSubMenu(
+				"Copy",
+				NSLOCTEXT("UserToolbox","Copy","Copy To ..."),
+				NSLOCTEXT("UserToolbox","CopyTooltip","Copy a command in another user toolbox tab"),
+				FNewMenuDelegate::CreateLambda([EditorPtr,CurrentCommand, AssetDatas](FMenuBuilder& MenuBuilder)
+				{
+					for (FAssetData AssetData:AssetDatas)
+					{
+						
+						MenuBuilder.AddSubMenu(FText::FromName(AssetData.AssetName),FText(),
+							FNewMenuDelegate::CreateLambda([EditorPtr,CurrentCommand,AssetData](FMenuBuilder& MenuBuilder)
+							{
+								UUserToolBoxBaseTab* CurrentTab=Cast<UUserToolBoxBaseTab>(AssetData.GetAsset());
+								if (IsValid(CurrentTab))
+								{
+									TArray<UUTBTabSection*> Sections=CurrentTab->GetSections();
+										Sections.Add(CurrentTab->GetPlaceHolderSection());
+									for (UUTBTabSection* TabSection:Sections)
+									{
+										MenuBuilder.AddMenuEntry(FText::FromString(TabSection->SectionName), FText(), FSlateIcon(),
+										FUIAction(FExecuteAction::CreateLambda([EditorPtr,CurrentCommand,TabSection,CurrentTab]()
+										{
+											CurrentCommand->CopyCommand(TabSection);
+											if (EditorPtr.IsValid())
+											{
+												TSharedPtr<FUTBTabEditor> TabEditor=EditorPtr.Pin();
+												TabEditor->NotifyTabHasChanged(CurrentTab);
+												TabEditor->UpdateCurrentSection();
+											}
+											
+											
+										})));
+									}
+								}
+							})
+							);
+					}
+					
+				}));
+		}));
+	}
+}
+TSharedPtr<SWidget> SCommandWidget::CreateContextMenu()
+{
+	RegisterContextMenu();
+	UToolMenus* ToolMenus = UToolMenus::Get();
+	FToolMenuContext NewContext(Command.Get());
+	UToolMenu* Menu = ToolMenus->GenerateMenu(DefaultContextBaseMenuName,NewContext);
+	for (const FToolMenuSection& Section : Menu->Sections)
+	{
+		if (Section.Blocks.Num() > 0)
+		{
+			return ToolMenus->GenerateWidget(Menu);
+		}
+	}
+
+	return nullptr;
+}
+FReply SCommandWidget::OnMouseButtonUp(const FGeometry& MyGeometry, const FPointerEvent& MouseEvent)
+{
+	if ( MouseEvent.GetEffectingButton() == EKeys::RightMouseButton && OnContextMenuOpening.IsBound() )
+	{
+		UE_LOG(LogTemp,Display,TEXT("MouseUp" ));
+		TSharedPtr<SWidget> ContextMenu = OnContextMenuOpening.Execute();
+		if ( ContextMenu )
+		{ 
+			FWidgetPath WidgetPath = MouseEvent.GetEventPath() ? *MouseEvent.GetEventPath() : FWidgetPath();
+			UE_LOG(LogTemp,Display,TEXT("Pushing Menu" ));
+			FSlateApplication::Get().PushMenu(
+				AsShared(),
+				WidgetPath,
+				ContextMenu.ToSharedRef(),
+				MouseEvent.GetScreenSpacePosition(),
+				FPopupTransitionEffect(FPopupTransitionEffect::ContextMenu)
+				);
+			return FReply::Handled();
+		}
+	}
+	return FReply::Unhandled();
+}
diff --git a/Engine/Plugins/Experimental/UserToolBoxCore/Source/UserToolBoxCore/Private/Editor/UTBEditorCommands.cpp b/Engine/Plugins/Experimental/UserToolBoxCore/Source/UserToolBoxCore/Private/Editor/UTBEditorCommands.cpp
new file mode 100644
--- /dev/null
+++ b/Engine/Plugins/Experimental/UserToolBoxCore/Source/UserToolBoxCore/Private/Editor/UTBEditorCommands.cpp
@@ -0,0 +1,30 @@
+// Copyright Epic Games, Inc. All Rights Reserved.
+
+
+#include "Editor/UTBEditorCommands.h"
+#include "Styling/AppStyle.h"
+#define LOCTEXT_NAMESPACE "UTBEditorCommands"
+FUTBEditorCommands::FUTBEditorCommands()
+	: TCommands<FUTBEditorCommands>
+(
+	TEXT("UsertoolboxEditor"),
+	LOCTEXT("RemoteControl", "Remote Control API"),
+	NAME_None,
+	FAppStyle::GetAppStyleSetName()
+	
+)
+{
+}
+
+FUTBEditorCommands::~FUTBEditorCommands()
+{
+}
+
+void FUTBEditorCommands::RegisterCommands()
+{
+	// Rename Entity
+	UI_COMMAND(RenameSection, "Rename", "Rename the selected  section.", EUserInterfaceActionType::Button, FInputChord(EKeys::F2));
+
+
+}
+#undef LOCTEXT_NAMESPACE
\ No newline at end of file
diff --git a/Engine/Plugins/Experimental/UserToolBoxCore/Source/UserToolBoxCore/Private/Editor/UTBEditorCommands.h b/Engine/Plugins/Experimental/UserToolBoxCore/Source/UserToolBoxCore/Private/Editor/UTBEditorCommands.h
new file mode 100644
--- /dev/null
+++ b/Engine/Plugins/Experimental/UserToolBoxCore/Source/UserToolBoxCore/Private/Editor/UTBEditorCommands.h
@@ -0,0 +1,25 @@
+// Copyright Epic Games, Inc. All Rights Reserved.
+
+#pragma once
+
+#include "CoreMinimal.h"
+#include "Framework/Commands/Commands.h"
+/**
+ * 
+ */
+class FUTBEditorCommands: public TCommands<FUTBEditorCommands>
+{
+public:
+	FUTBEditorCommands();
+	~FUTBEditorCommands();
+
+	//~ BEGIN : TCommands<> Implementation(s)
+
+	virtual void RegisterCommands() override;
+
+	//~ END : TCommands<> Implementation(s)
+
+
+	
+	TSharedPtr< FUICommandInfo > RenameSection;
+};
diff --git a/Engine/Plugins/Experimental/UserToolBoxCore/Source/UserToolBoxCore/Private/Editor/UTBTabEditor.cpp b/Engine/Plugins/Experimental/UserToolBoxCore/Source/UserToolBoxCore/Private/Editor/UTBTabEditor.cpp
new file mode 100644
--- /dev/null
+++ b/Engine/Plugins/Experimental/UserToolBoxCore/Source/UserToolBoxCore/Private/Editor/UTBTabEditor.cpp
@@ -0,0 +1,917 @@
+// Copyright Epic Games, Inc. All Rights Reserved.
+
+
+#include "Editor/UTBTabEditor.h"
+
+
+#include "SAssetDropTarget.h"
+#include "AssetRegistry/AssetRegistryModule.h"
+#include "Editor/SCommandPickMenu.h"
+#include "Editor/SCommandWithOptionsWidget.h"
+
+#include "Styling/StyleColors.h"
+#include "SDropTarget.h"
+#include "ToolMenuDelegates.h"
+#include "ToolMenus.h"
+#include "UserToolBoxSubsystem.h"
+#include "UTBEditorCommands.h"
+#include "DragAndDrop/DecoratedDragDropOp.h"
+#include "Interfaces/IMainFrameModule.h"
+#include "Widgets/Layout/SScrollBox.h"
+#include "Widgets/Layout/SUniformGridPanel.h"
+#include "Widgets/Layout/SWrapBox.h"
+#include "Widgets/Text/SInlineEditableTextBlock.h"
+#include "Widgets/Docking/SDockTab.h"
+#include "PropertyEditorModule.h"
+#define LOCTEXT_NAMESPACE "UserToolboxTabEditor"
+
+const FName FUTBTabEditor::UTBTabTabId ("UTBTabEditor_Tab");
+const FName FUTBTabEditor::UTBTabSectionTabId ("UTBTabEditor_Section");
+const FName FUTBTabEditor::UTBCommandTabId("UTBTabEditor_Command");
+const FName FUTBTabEditor::UTBCommandDetailsTabId("UTBTabEditor_CommandDetails");
+const FName FUTBTabEditor::UTBTabDetailsTabId("UTBTabEditor_Details");
+const FName UTBTabEditorAppIdentifier( TEXT( "UTBTabEditorApp" ) );
+static const FName DefaultContextSectionBaseMenuName("UserToolboxTabEditor.SectionDefaultContextMenuBase");
+
+
+class FSectionDragDrop final : public FDecoratedDragDropOp
+{
+public:
+	DRAG_DROP_OPERATOR_TYPE(FSectionDragDrop, FDecoratedDragDropOp)
+
+	using WidgetType = SWidget;
+
+	FSectionDragDrop(TSharedPtr<SWidget> InWidget, const FString& SectionName)
+		: SectionName(SectionName)
+	{
+		
+	}
+
+	/** Get the ID of the represented entity or group. */
+	FString GetSectionName() const
+	{
+		return SectionName;
+	}
+
+	virtual void OnDrop(bool bDropWasHandled, const FPointerEvent& MouseEvent) override
+	{
+		FDecoratedDragDropOp::OnDrop(bDropWasHandled, MouseEvent);
+	}
+
+
+private:
+	FString SectionName;
+	TSharedPtr<SWidget> DecoratorWidget;
+};
+
+
+class SSectionDragHandle : public SCompoundWidget
+{
+public:
+	
+	SLATE_BEGIN_ARGS(SSectionDragHandle)
+	{}
+	SLATE_ARGUMENT(TSharedPtr<SWidget>, Widget)
+SLATE_END_ARGS()
+
+void Construct(const FArguments& InArgs, FString _SectionName)
+	{
+		Widget = InArgs._Widget;
+		SectionName=_SectionName;
+		ChildSlot
+		[
+			SNew(SBox)
+			.WidthOverride(25.0f)
+			.Visibility(SectionName!=UUserToolBoxBaseTab::PlaceHolderSectionName.ToString()?EVisibility::Visible:EVisibility::Collapsed)
+			[
+				SNew(SHorizontalBox)
+				+ SHorizontalBox::Slot()
+				.Padding(5.f, 0.f)
+				.VAlign(VAlign_Center)
+				.AutoWidth()
+				[
+					SNew(SImage)
+					.Image(FCoreStyle::Get().GetBrush("VerticalBoxDragIndicatorShort"))
+				]
+			]
+		];
+	}
+
+	FReply OnMouseButtonDown(const FGeometry& MyGeometry, const FPointerEvent& MouseEvent) override
+	{
+		return FReply::Handled().DetectDrag(SharedThis(this), EKeys::LeftMouseButton);
+	};
+
+	FReply OnDragDetected(const FGeometry& MyGeometry, const FPointerEvent& MouseEvent) override
+	{
+		if (MouseEvent.IsMouseButtonDown(EKeys::LeftMouseButton))
+		{
+			TSharedPtr<FDragDropOperation> DragDropOp = CreateDragDropOperation();
+			if (DragDropOp.IsValid())
+			{
+				return FReply::Handled().BeginDragDrop(DragDropOp.ToSharedRef());
+			}
+		}
+
+		return FReply::Unhandled();
+	}
+
+	TSharedPtr<FDragDropOperation> CreateDragDropOperation()
+	{
+		TSharedPtr<FSectionDragDrop> DragDropOperation = MakeShared<FSectionDragDrop>(Widget.Pin(), SectionName);
+		DragDropOperation->Construct();
+		return DragDropOperation;
+		
+	}
+
+private:
+	/** Holds the widget to display when dragging. */
+	TWeakPtr<SWidget> Widget;
+	/** Holds the ID of the item being dragged. */
+	FString SectionName;
+};
+class SSectionRowWidget : public STableRow<TSharedPtr<FString>>
+{
+public:
+	
+	SLATE_BEGIN_ARGS(SSectionRowWidget)
+	{}
+	SLATE_ARGUMENT(FString, SectionName)
+	SLATE_ARGUMENT(TSharedPtr<STableViewBase>,	Owner)
+	SLATE_EVENT(FOnDrop, OnDropped)
+	SLATE_EVENT(SDropTarget::FVerifyDrag, OnAllowDrop)
+		/** Callback when the text is committed. */
+	SLATE_EVENT( FOnTextCommitted, OnTextCommitted )
+SLATE_END_ARGS()
+
+	virtual FReply OnKeyDown(const FGeometry& MyGeometry, const FKeyEvent& InKeyEvent) override
+	{
+		if (InKeyEvent.GetKey()==EKeys::F2)
+		{
+			EnableRenaming();
+			return FReply::Handled();
+		}
+		return FReply::Unhandled();
+	}
+
+	void Construct(const FArguments& InArgs, FString _SectionName, const TSharedRef<STableViewBase>& InOwnerTableView)
+	{
+		SectionName=_SectionName;
+		STableRow<TSharedPtr<FString>>::Construct(STableRow<TSharedPtr<FString>>::FArguments().Content()
+			[
+			SNew(SBorder)
+						.Padding(0.f)
+						.VAlign(VAlign_Fill)
+						[
+							SNew(SBox)
+							.HeightOverride(32.f)
+							[
+								SNew(SDropTarget)
+								.OnDropped(InArgs._OnDropped)
+								.OnAllowDrop(InArgs._OnAllowDrop)
+								[
+									SNew(SHorizontalBox)
+									+ SHorizontalBox::Slot()
+									.VAlign(VAlign_Fill)
+									.HAlign(HAlign_Center)
+									.Padding(0.f, 4.f)
+									.AutoWidth()
+									[
+										SNew(SSectionDragHandle,*SectionName)
+										.Widget(SharedThis(this))
+									]
+									+ SHorizontalBox::Slot()
+									// Group name
+									.FillWidth(1.f)
+									.VAlign(VAlign_Center)
+									.Padding(6.f, 4.f)
+									.AutoWidth()
+											
+									[
+										SAssignNew(EditableTextBlock,SInlineEditableTextBlock)
+										.Font(FAppStyle::Get().GetFontStyle("PropertyWindow.BoldFont"))
+										.Text(FText::FromString(*SectionName))
+										.OnTextCommitted(InArgs._OnTextCommitted)
+										.IsReadOnly(false)
+									]
+									
+
+								]
+							]
+						]
+			],InOwnerTableView);
+	}
+
+	void EnableRenaming()
+	{
+		if (EditableTextBlock.IsValid())
+		{
+			EditableTextBlock->EnterEditingMode();
+		}
+		
+	}
+
+private:
+	/** Holds the ID of the item being dragged. */
+	FString SectionName;
+	TSharedPtr<SInlineEditableTextBlock> EditableTextBlock;
+};
+
+FCommandDragDropOperation::FCommandDragDropOperation():Command()
+{
+}
+
+FCommandDragDropOperation::~FCommandDragDropOperation()
+{
+}
+
+TSharedRef<FCommandDragDropOperation> FCommandDragDropOperation::New(UUTBBaseCommand* Command)
+{
+	TSharedRef<FCommandDragDropOperation> NewOp=MakeShareable(new FCommandDragDropOperation());
+	NewOp->Command=Command;
+	return NewOp;
+}
+
+FUTBTabEditor::FUTBTabEditor()
+{
+	
+}
+
+FUTBTabEditor::~FUTBTabEditor()
+{
+}
+
+FLinearColor FUTBTabEditor::GetWorldCentricTabColorScale() const
+{
+	return FLinearColor( 0.0f, 0.0f, 0.2f, 0.5f );
+}
+
+void FUTBTabEditor::RegisterTabSpawners(const TSharedRef<FTabManager>& InTabManager)
+{
+	IUTBTabEditor::RegisterTabSpawners(InTabManager);
+	CreateOrRebuildUTBTabTabWidget();
+	CreateOrRebuildUTBTabSectionTabWidget();
+	UUTBTabDetailsTabWidget=CreateUTBTabDetailsTabWidget();
+	UUTBCommandTabWidget=CreateUTBCommandTabWidget();
+	UUTBCommandDetailsTabWidget=CreateUTBCommandDetailsTabWidget();
+	
+	InTabManager->RegisterTabSpawner(UTBTabTabId,FOnSpawnTab::CreateLambda([this](const FSpawnTabArgs& Args)
+	{
+		return SNew(SDockTab)
+		.Label( LOCTEXT("UserToolBoxTab", "User Toolbox Tab") )
+		.TabColorScale( GetTabColorScale() )
+		[
+			SNew(SBorder)
+			.Padding(2)
+			.BorderImage( FAppStyle::GetBrush( "ToolPanel.GroupBorder" ) )
+			[
+				UUTBTabTabWidget.ToSharedRef()
+			]
+		];
+	}))
+	.SetDisplayName(LOCTEXT("UserToolBoxTab","User Toolbox Tab"));
+	
+	InTabManager->RegisterTabSpawner(UTBTabSectionTabId,FOnSpawnTab::CreateLambda([this](const FSpawnTabArgs& Args)
+	{
+		return SNew(SDockTab)
+		.Label( LOCTEXT("UserToolBoxTabSection", "Tab Section") )
+		.TabColorScale( GetTabColorScale() )
+		[
+			SNew(SBorder)
+			.Padding(2)
+			.BorderImage( FAppStyle::GetBrush( "ToolPanel.GroupBorder" ) )
+			[
+				UUTBTabSectionTabWidget.ToSharedRef()
+			]
+		];
+	}))
+	.SetDisplayName(LOCTEXT("UserToolBoxTabSection","Tab Section"));
+	
+	 InTabManager->RegisterTabSpawner(UTBTabDetailsTabId,FOnSpawnTab::CreateLambda([this](const FSpawnTabArgs& Args)
+	 {
+	 	return SNew(SDockTab)
+	 	.Label( LOCTEXT("UserToolBoxTabDetails", "Tab Details") )
+	 	.TabColorScale( GetTabColorScale() )
+	 	[
+	 		SNew(SBorder)
+	 		.Padding(2)
+	 		.BorderImage( FAppStyle::GetBrush( "ToolPanel.GroupBorder" ) )
+	 		[
+	 			UUTBTabDetailsTabWidget.ToSharedRef()
+	 		]
+	 	];
+	 }))
+	 .SetDisplayName(LOCTEXT("UserToolBoxTabDetails","Tab Details"));
+	
+	InTabManager->RegisterTabSpawner(UTBCommandTabId,FOnSpawnTab::CreateLambda([this](const FSpawnTabArgs& Args)
+	{
+		return SNew(SDockTab)
+		.Label( LOCTEXT("UserToolBoxTabCommandList", "Command List") )
+		.TabColorScale( GetTabColorScale() )
+		[
+			SNew(SBorder)
+			.Padding(2)
+			.BorderImage( FAppStyle::GetBrush( "ToolPanel.GroupBorder" ) )
+			[
+				UUTBCommandTabWidget.ToSharedRef()
+			]
+		];
+	}))
+	.SetDisplayName(LOCTEXT("UserToolBoxTabCommandList","Command List"));
+	
+	InTabManager->RegisterTabSpawner(UTBCommandDetailsTabId,FOnSpawnTab::CreateLambda([this](const FSpawnTabArgs& Args)
+	{
+		return SNew(SDockTab)
+		.Label( LOCTEXT("UserToolBoxCommandDetails", "Command Details") )
+		.TabColorScale( GetTabColorScale() )
+		[
+			SNew(SBorder)
+			.Padding(2)
+			.BorderImage( FAppStyle::GetBrush( "ToolPanel.GroupBorder" ) )
+			[
+				UUTBCommandDetailsTabWidget.ToSharedRef()
+			]
+		];
+	}))
+	.SetDisplayName(LOCTEXT("UserToolBoxCommandDetails","Command Details"));
+	
+}
+
+void FUTBTabEditor::UnregisterTabSpawners(const TSharedRef<FTabManager>& InTabManager)
+{
+	IUTBTabEditor::UnregisterTabSpawners(InTabManager);
+	InTabManager->UnregisterTabSpawner(UTBTabTabId);
+	InTabManager->UnregisterTabSpawner(UTBTabSectionTabId);
+	InTabManager->UnregisterTabSpawner(UTBTabDetailsTabId);
+	InTabManager->UnregisterTabSpawner(UTBCommandTabId);
+	InTabManager->UnregisterTabSpawner(UTBCommandDetailsTabId);
+}
+
+FName FUTBTabEditor::GetToolkitFName() const
+{
+	return FName("UsertoolboxTabEditor");
+}
+
+FText FUTBTabEditor::GetBaseToolkitName() const
+{
+	return LOCTEXT( "AppLabel", "User Toolbox Tab Editor" );
+}
+
+FString FUTBTabEditor::GetWorldCentricTabPrefix() const
+{
+	return "UserToolboxTab";
+}
+
+void FUTBTabEditor::PostUndo(bool bSuccess)
+{
+	FEditorUndoClient::PostUndo(bSuccess);
+}
+
+void FUTBTabEditor::PostRedo(bool bSuccess)
+{
+	FEditorUndoClient::PostRedo(bSuccess);
+}
+
+void FUTBTabEditor::InitUTBTabEditor(const EToolkitMode::Type Mode, const TSharedPtr<IToolkitHost>& InitToolkitHost,UUserToolBoxBaseTab* Tab)
+{
+	TSharedRef<FTabManager::FLayout> StandaloneDefaultLayout = FTabManager::NewLayout( "Standalone_UserToolBoxEditor_v1" )
+	->AddArea
+	(
+		FTabManager::NewPrimaryArea()->SetOrientation(Orient_Horizontal)
+		->Split
+		(
+			FTabManager::NewStack()
+			->SetSizeCoefficient(0.33f)
+			->AddTab(UTBTabSectionTabId, ETabState::OpenedTab)
+		)
+		->Split
+		(
+			FTabManager::NewStack()
+			->SetSizeCoefficient(0.9f)
+			->AddTab(UTBTabTabId, ETabState::OpenedTab)
+		)
+		->Split
+		(
+			FTabManager::NewStack()
+			->SetSizeCoefficient(0.66f)
+			->AddTab(UTBTabDetailsTabId, ETabState::OpenedTab)
+			->AddTab(UTBCommandTabId, ETabState::OpenedTab)
+			->SetForegroundTab(UTBCommandTabId)
+					
+		)
+		->Split
+		(
+		FTabManager::NewStack()
+			->SetSizeCoefficient(0.33f)
+			
+			->AddTab(UTBCommandDetailsTabId, ETabState::OpenedTab)
+
+		)
+	);
+
+	const bool bCreateDefaultStandaloneMenu = true;
+	const bool bCreateDefaultToolbar = true;
+	FAssetEditorToolkit::InitAssetEditor( Mode, InitToolkitHost, UTBTabEditorAppIdentifier, StandaloneDefaultLayout, bCreateDefaultStandaloneMenu, bCreateDefaultToolbar, Tab );
+
+	RegenerateMenusAndToolbars();
+
+	// Support undo/redo
+	GEditor->RegisterForUndo(this);
+	BindCommands();
+}
+
+FString FUTBTabEditor::GetCurrentSection()
+{
+	return CurrentSection;
+}
+
+void FUTBTabEditor::AddCommandToCurrentSection(UUTBBaseCommand* Command)
+{
+	UUserToolBoxBaseTab* Tab=Cast<UUserToolBoxBaseTab>(GetEditingObject());
+	if (!IsValid(Tab))
+	{
+		return ;
+	}
+		
+	if (!IsValid(Command))
+	{
+		return;
+	}
+	AddCommandAt(Command,CurrentSection,-1);
+}
+
+void FUTBTabEditor::RemoveCommandFromCurrentSection(UUTBBaseCommand* InCommand,bool ShouldRebuild)
+{
+	UUserToolBoxBaseTab* Tab=Cast<UUserToolBoxBaseTab>(GetEditingObject());
+	if (!IsValid(Tab) || !IsValid(InCommand))
+		return ;
+	Tab->RemoveCommandFromSection(InCommand,CurrentSection);
+	
+	bool ret=Tab->MarkPackageDirty();
+	if (ShouldRebuild)
+	{
+		CreateOrRebuildUTBTabTabWidget();
+	}
+	GEditor->GetEditorSubsystem<UUserToolboxSubsystem>()->OnTabChanged.Broadcast(Tab);
+}
+
+void FUTBTabEditor::AddCommandAt(UUTBBaseCommand* InCommand,FString Section, int Index)
+{
+	UUserToolBoxBaseTab* Tab=Cast<UUserToolBoxBaseTab>(GetEditingObject());
+	if (!IsValid(Tab) || !IsValid(InCommand))
+		return ;
+	Tab->InsertCommand(InCommand,Section,Index);
+	bool ret=Tab->MarkPackageDirty();
+	CreateOrRebuildUTBTabTabWidget();
+	
+	GEditor->GetEditorSubsystem<UUserToolboxSubsystem>()->OnTabChanged.Broadcast(Tab);
+}
+
+
+
+void FUTBTabEditor::SetCurrentCommandSelection(UUTBBaseCommand* Command)
+{
+	CommandDetailsView->SetObject(Command);
+}
+
+void FUTBTabEditor::NotifyTabHasChanged(UUserToolBoxBaseTab* Tab)
+{
+	if (!IsValid(Tab))
+	{
+		return;
+	}
+	GEditor->GetEditorSubsystem<UUserToolboxSubsystem>()->OnTabChanged.Broadcast(Cast<UUserToolBoxBaseTab>(Tab));
+	if (Tab==GetEditingObject())
+	{
+		CreateOrRebuildUTBTabTabWidget();
+	}
+}
+
+void FUTBTabEditor::UpdateCommand(UUTBBaseCommand* Command)
+{
+
+
+	
+}
+
+void FUTBTabEditor::UpdateCurrentSection()
+{
+	CreateOrRebuildUTBTabTabWidget();
+}
+
+void FUTBTabEditor::CreateNewSection()
+{
+	UUserToolBoxBaseTab* Tab=Cast<UUserToolBoxBaseTab>(GetEditingObject());
+	if (!IsValid(Tab))
+	{
+		return ;
+	}
+	int Index=0;
+	FString Name="New Section";
+	FString TmpName="New Section";
+	
+	while (Tab->GetSection(TmpName))
+	{
+		TmpName=Name+"_"+FString::FromInt(Index);
+		Index++;
+	}
+	Tab->GetOrCreateSection(TmpName);
+	RefreshSectionNames();
+	GEditor->GetEditorSubsystem<UUserToolboxSubsystem>()->OnTabChanged.Broadcast(Tab);
+}
+
+void FUTBTabEditor::MoveSectionAfterExistingSection(FString SectionToMoveName, FString SectionToTargetName)
+{
+	UUserToolBoxBaseTab* Tab=Cast<UUserToolBoxBaseTab>(GetEditingObject());
+	if (!IsValid(Tab))
+	{
+		return ;
+	}
+	Tab->MoveSectionAfterExistingSection(SectionToMoveName,SectionToTargetName);
+	RefreshSectionNames();
+	GEditor->GetEditorSubsystem<UUserToolboxSubsystem>()->OnTabChanged.Broadcast(Tab);
+}
+
+void FUTBTabEditor::RefreshSectionNames()
+{
+	UUserToolBoxBaseTab* Tab=Cast<UUserToolBoxBaseTab>(GetEditingObject());
+	if (!IsValid(Tab))
+	{
+		return ;
+	}
+	SectionNames.Empty();
+	for (UUTBTabSection* Section:Tab->GetSections())
+	{
+		if (!IsValid(Section))
+		{
+			continue;
+		}
+		SectionNames.Add(MakeShared<FString>(Section->SectionName));
+	}
+	SectionNames.Add(MakeShared<FString>(UUserToolBoxBaseTab::PlaceHolderSectionName.ToString()));
+	if (SectionListView.IsValid())
+	{
+		SectionListView->RequestListRefresh();
+	}
+}
+
+void FUTBTabEditor::RemoveCurrentSection()
+{
+	UUserToolBoxBaseTab* Tab=Cast<UUserToolBoxBaseTab>(GetEditingObject());
+	if (!IsValid(Tab))
+	{
+		return ;
+	}
+	FString CurrentSectionName=GetCurrentSection();
+	if (CurrentSectionName.IsEmpty())
+	{
+		return;
+	}
+	Tab->RemoveSection(CurrentSectionName);
+	GEditor->GetEditorSubsystem<UUserToolboxSubsystem>()->OnTabChanged.Broadcast(Tab);
+	RefreshSectionNames();
+}
+
+void FUTBTabEditor::EditActiveSectionName()
+{
+	if (SectionListView.IsValid())
+	{
+		for (TSharedPtr<FString> CurrentSectionName:SectionNames)
+		{
+			if (*CurrentSectionName==GetCurrentSection())
+			{
+				TSharedPtr<ITableRow> TableRow= SectionListView->WidgetFromItem(CurrentSectionName);
+				if (TableRow.IsValid())
+				{
+					TSharedPtr<SSectionRowWidget> AsSectionRow=StaticCastSharedPtr<SSectionRowWidget>(TableRow);
+					if (AsSectionRow.IsValid())
+					{
+						AsSectionRow->EnableRenaming();
+					}
+				}
+			}
+		}
+		
+	}
+}
+
+void FUTBTabEditor::RenameActiveSection(FString NewName)
+{
+	UUserToolBoxBaseTab* Tab=Cast<UUserToolBoxBaseTab>(GetEditingObject());
+	if (!IsValid(Tab) || CurrentSection.IsEmpty() || CurrentSection==NewName || CurrentSection==UUserToolBoxBaseTab::PlaceHolderSectionName.ToString())
+	{
+		return ;
+	}
+	
+	UUTBTabSection* SectionPtr=Tab->GetSection(CurrentSection);
+	if (!IsValid(SectionPtr))
+	{
+		return;
+	}
+	SectionPtr->SectionName=NewName;
+	Tab->MarkPackageDirty();
+	CurrentSection=NewName;
+	RefreshSectionNames();
+	GEditor->GetEditorSubsystem<UUserToolboxSubsystem>()->OnTabChanged.Broadcast(Tab);
+}
+
+void FUTBTabEditor::BindCommands()
+{
+	const FUTBEditorCommands& Commands = FUTBEditorCommands::Get();
+
+	IMainFrameModule& MainFrame = FModuleManager::Get().LoadModuleChecked<IMainFrameModule>("MainFrame");
+
+	FUICommandList& ActionList = *MainFrame.GetMainFrameCommandBindings();
+
+	ActionList.MapAction(
+		Commands.RenameSection,
+		FExecuteAction::CreateSP(this, &FUTBTabEditor::EditActiveSectionName),
+		FCanExecuteAction::CreateLambda([this]()
+		{
+			return CurrentSection!=UUserToolBoxBaseTab::PlaceHolderSectionName.ToString();
+		}));
+
+}
+
+inline void FUTBTabEditor::CreateOrRebuildUTBTabTabWidget()
+{
+	UUserToolBoxBaseTab* Tab=Cast<UUserToolBoxBaseTab>(GetEditingObject());
+	if (!IsValid(Tab))
+	{
+		return ;
+	}
+	if (UUTBTabTabWidget.IsValid())
+	{
+		UUTBTabTabWidget->ClearChildren();
+	}
+	else
+	{
+		SAssignNew(UUTBTabTabWidget,SVerticalBox)	;
+	}
+	if (Tab->GetSection(CurrentSection)==nullptr)
+	{
+		UUTBTabTabWidget->AddSlot()
+		[
+			SNew(STextBlock)
+			.Text(FText::FromString("Select a valid section"))
+		];
+	}
+	else
+	{
+		UUTBTabTabWidget->AddSlot()
+		.VAlign(VAlign_Top)
+		.AutoHeight()
+		.HAlign(HAlign_Fill)
+		[
+			SNew(SBorder)
+			.HAlign(HAlign_Fill)
+			.VAlign(VAlign_Fill)
+			[
+				SNew(STextBlock)
+				.Text(FText::FromString(CurrentSection))
+			]
+		];
+
+		
+		
+		TSharedPtr<SWrapBox> PanelWidget= SNew(SWrapBox)
+		.UseAllottedSize(true);
+		TArray<UUTBBaseCommand*>& CurrentSectionArray=Tab->GetOrCreateSection(CurrentSection)->Commands;
+		int Index=0;
+		for (UUTBBaseCommand* Data:CurrentSectionArray)
+		{
+			PanelWidget->AddSlot()
+			.Padding(10)
+			[
+				SNew(SCommandWidget)
+				.Command(Data)
+				.SectionName(CurrentSection)
+				.IndexInSection(Index)
+				.Editor(SharedThis(this))
+				.OnTabChanged_Raw(this, &FUTBTabEditor::NotifyTabHasChanged)
+				.OnCommandSelected_Raw(this,&FUTBTabEditor::SetCurrentCommandSelection)
+				.OnCommandDeleted_Lambda([this,Tab](UUTBBaseCommand* Command)
+				{
+					if (IsValid(Tab))
+					{
+						Tab->RemoveCommand(Command);
+						NotifyTabHasChanged(Tab);
+						CreateOrRebuildUTBTabTabWidget();
+					}
+				})
+			];
+			Index++;
+		}
+		UUTBTabTabWidget->AddSlot()
+		.VAlign(VAlign_Fill)
+		[
+			SNew(SDropTarget)
+			.OnAllowDrop(SDropTarget::FVerifyDrag::CreateLambda([](TSharedPtr<FDragDropOperation> Operation)
+			{
+
+				return Operation->IsOfType<FCommandDragDropOperation>();
+			}))
+			.OnDropped(FOnDrop::CreateLambda([this](const FGeometry& Geometry,const FDragDropEvent& Event)
+			{
+				if ( Event.GetOperation() && Event.GetOperation()->IsOfType<FCommandDragDropOperation>())
+				{
+					if (!Event.GetOperationAs<FCommandDragDropOperation>()->Command.IsValid())
+					{
+						return FReply::Handled().EndDragDrop();
+					}
+					UUTBBaseCommand* Command=Event.GetOperationAs<FCommandDragDropOperation>()->Command.Get();
+					RemoveCommandFromCurrentSection(Command);
+					this->AddCommandToCurrentSection(Command);
+				}
+			
+			return FReply::Handled().EndDragDrop();	
+			}))
+			[
+				PanelWidget.ToSharedRef()
+			]
+		];
+	}
+	
+}
+void FUTBTabEditor::CreateOrRebuildUTBTabSectionTabWidget()
+{
+	UUserToolBoxBaseTab* Tab=Cast<UUserToolBoxBaseTab>(GetEditingObject());
+	if (!IsValid(Tab))
+	{
+		return ;
+	}
+	if (UUTBTabSectionTabWidget.IsValid())
+	{
+		UUTBTabSectionTabWidget->ClearChildren();
+	}
+	else
+	{
+		SAssignNew(UUTBTabSectionTabWidget,SVerticalBox)	;
+	}
+	UUTBTabSectionTabWidget->AddSlot()
+	.AutoHeight()
+   [
+	   SNew(SHorizontalBox)
+	   +SHorizontalBox::Slot()
+	   [
+			SNew(SHorizontalBox)
+			+SHorizontalBox::Slot()
+			.Padding(5.f, 0.f)
+			.AutoWidth()
+			.VAlign(VAlign_Center)
+			[
+				SNew(SButton)
+				.AddMetaData<FTagMetaData>(FTagMetaData(TEXT("Add Section")))
+				.HAlign(HAlign_Center)
+				.VAlign(VAlign_Center)
+				.ForegroundColor(FSlateColor::UseForeground())
+				.ToolTipText(LOCTEXT("NewGroupToolTip", "Create new group."))
+				.OnClicked_Lambda([this]()
+				{
+					CreateNewSection();
+					return FReply::Handled();
+				})
+				
+				[
+					SNew(SBox)
+					[
+						SNew(SImage)
+						.ColorAndOpacity(FSlateColor::UseForeground())
+						.Image(FAppStyle::GetBrush("Icons.PlusCircle"))
+					]
+				]
+			]
+	   ]
+	   +SHorizontalBox::Slot()
+	   .FillWidth(1.f)
+		.VAlign(VAlign_Center)
+		.Padding(2.f, 4.f)
+		[
+		   SNew(STextBlock)
+		   .Text(FText::FromString("Sections"))
+		]
+		+SHorizontalBox::Slot()
+		.AutoWidth()
+		.Padding(2.f, 4.f, 4.f, 4.f)
+		.VAlign(VAlign_Center)
+		[
+			SNew(SButton)
+			.AddMetaData<FTagMetaData>(FTagMetaData(TEXT("Delete Selected Section")))
+			.IsEnabled_Lambda([this]() { return !this->CurrentSection.IsEmpty() && CurrentSection!=UUserToolBoxBaseTab::PlaceHolderSectionName.ToString();})
+			.HAlign(HAlign_Center)
+			.VAlign(VAlign_Center)
+			.ForegroundColor(FSlateColor::UseForeground())
+			.ToolTipText(LOCTEXT("DeletedSelectedSection", "Delete selected section."))
+			.OnClicked_Lambda([this]()
+			{
+				
+				RemoveCurrentSection();
+				return FReply::Handled();
+			})
+			[
+				SNew(SBox)
+				[
+					SNew(SImage)
+					.ColorAndOpacity(FSlateColor::UseForeground())
+					.Image(FAppStyle::GetBrush("Icons.Delete"))
+				]
+			]
+		]
+
+   ];
+	RefreshSectionNames();		
+	UUTBTabSectionTabWidget->AddSlot()
+	.AutoHeight()
+	[		// Groups List
+		SAssignNew( SectionListView,SListView<TSharedPtr<FString>>)
+			.ItemHeight(24.f)
+			.OnGenerateRow_Lambda(
+
+			[this] (const TSharedPtr<FString> SectionName, const TSharedRef<STableViewBase>& OwnerTable)
+				{
+					return SNew(SSectionRowWidget,*SectionName,OwnerTable)
+						.OnDropped_Lambda([this,SectionName,OwnerTable](const FGeometry& InGeometry, const FDragDropEvent& InDragDropEvent)
+							{
+								if (const TSharedPtr<FSectionDragDrop> DragDropOp = InDragDropEvent.GetOperationAs<FSectionDragDrop>())
+								{
+									FString DroppedSectionName=DragDropOp->GetSectionName();
+									this->MoveSectionAfterExistingSection(DroppedSectionName,*SectionName);
+									OwnerTable->RequestListRefresh();
+									return FReply::Handled();
+								}
+								return FReply::Unhandled();
+							})
+							.OnAllowDrop_Lambda([SectionName](TSharedPtr<FDragDropOperation> DragDropOperation)
+							{
+								return DragDropOperation->IsOfType<FSectionDragDrop>() && !SectionName->Equals(UUserToolBoxBaseTab::PlaceHolderSectionName.ToString());
+							})
+						.OnTextCommitted_Lambda([this]( const FText& Text, ETextCommit::Type)
+						{
+							RenameActiveSection(Text.ToString());
+						})
+
+				;
+				}
+				
+			)
+			.OnSelectionChanged_Lambda([this](TSharedPtr<FString> SectionName,	ESelectInfo::Type Type)
+			{
+				if (SectionName.IsValid())
+				{
+					SetCurrentSection(*SectionName);
+				}
+				else
+				{
+					SetCurrentSection("");
+				}
+			})
+			.SelectionMode(ESelectionMode::Single)
+			.ListItemsSource(&SectionNames)
+			.OnContextMenuOpening_Lambda([this]()
+			{
+				IMainFrameModule& MainFrame = FModuleManager::Get().LoadModuleChecked<IMainFrameModule>("MainFrame");
+					FMenuBuilder MenuBuilder(true, MainFrame.GetMainFrameCommandBindings());
+					MenuBuilder.BeginSection("Common");
+					MenuBuilder.AddMenuEntry(FUTBEditorCommands::Get().RenameSection, NAME_None, TAttribute<FText>(), TAttribute<FText>(), FSlateIcon(FAppStyle::GetAppStyleSetName(), TEXT("GenericCommands.Rename")));
+					MenuBuilder.EndSection();
+					return MenuBuilder.MakeWidget();
+			})
+			.ClearSelectionOnClick(true)
+	];
+}
+inline TSharedPtr<SWidget> FUTBTabEditor::CreateUTBTabDetailsTabWidget()
+{
+	FPropertyEditorModule& PropertyEditorModule = FModuleManager::GetModuleChecked<FPropertyEditorModule>("PropertyEditor");
+	FDetailsViewArgs Args;
+	Args.bAllowSearch=true;
+	Args.bShowOptions=true;
+	Args.bHideSelectionTip=true;
+	Args.bShowKeyablePropertiesOption=false;
+	Args.bShowObjectLabel=false;
+
+	Args.bUpdatesFromSelection=false;
+	Args.NameAreaSettings=FDetailsViewArgs::ObjectsUseNameArea;
+	TSharedRef<IDetailsView> DetailView=PropertyEditorModule.CreateDetailView(Args);
+	DetailView->SetObject(Cast<UUserToolBoxBaseTab>(GetEditingObject()));
+	return 	DetailView;
+}
+
+inline TSharedPtr<SWidget> FUTBTabEditor::CreateUTBCommandTabWidget()
+{
+	return SNew(SCommandPickMenuWidget);
+}
+
+TSharedPtr<SWidget> FUTBTabEditor::CreateUTBCommandDetailsTabWidget()
+{
+	return SAssignNew(CommandDetailsView,SCommandDetailsView)
+	.OnObjectPropertyModified_Raw(this,&FUTBTabEditor::UpdateCommand);
+}
+
+void FUTBTabEditor::SetCurrentSection(FString SectionName)
+{
+	CurrentSection=SectionName;
+	CreateOrRebuildUTBTabTabWidget();
+	OnCurrentSectionChange.Broadcast(CurrentSection);
+}
+
+#undef LOCTEXT_NAMESPACE 
diff --git a/Engine/Plugins/Experimental/UserToolBoxCore/Source/UserToolBoxCore/Private/IconsTracker.cpp b/Engine/Plugins/Experimental/UserToolBoxCore/Source/UserToolBoxCore/Private/IconsTracker.cpp
new file mode 100644
--- /dev/null
+++ b/Engine/Plugins/Experimental/UserToolBoxCore/Source/UserToolBoxCore/Private/IconsTracker.cpp
@@ -0,0 +1,5 @@
+// Copyright Epic Games, Inc. All Rights Reserved.
+
+
+#include "IconsTracker.h"
+
diff --git a/Engine/Plugins/Experimental/UserToolBoxCore/Source/UserToolBoxCore/Private/SVerticalTextBlock.cpp b/Engine/Plugins/Experimental/UserToolBoxCore/Source/UserToolBoxCore/Private/SVerticalTextBlock.cpp
new file mode 100644
--- /dev/null
+++ b/Engine/Plugins/Experimental/UserToolBoxCore/Source/UserToolBoxCore/Private/SVerticalTextBlock.cpp
@@ -0,0 +1,55 @@
+// Copyright Epic Games, Inc. All Rights Reserved.
+
+
+#include "SVerticalTextBlock.h"
+
+SVerticalTextBlock::SVerticalTextBlock()
+{
+}
+
+SVerticalTextBlock::~SVerticalTextBlock()
+{
+}
+
+void SVerticalTextBlock::Construct(const FArguments& InArgs)
+{
+	STextBlock::Construct(STextBlock::FArguments()
+		.Font(InArgs._Font)
+		.Justification(InArgs._Justification)
+		.Margin(InArgs._Margin)
+		.Text(InArgs._Text)
+		.HighlightColor(InArgs._HighlightColor)
+		.HighlightShape(InArgs._HighlightShape)
+		.HighlightText(InArgs._HighlightText)
+		.OverflowPolicy(InArgs._OverflowPolicy)
+		.ShadowOffset(InArgs._ShadowOffset)
+		.StrikeBrush(InArgs._StrikeBrush)
+		.TextStyle(InArgs._TextStyle)
+		.TransformPolicy(InArgs._TransformPolicy)
+		.WrappingPolicy(InArgs._WrappingPolicy)
+		.AutoWrapText(InArgs._AutoWrapText)
+		.ColorAndOpacity(InArgs._ColorAndOpacity)
+		.WrapTextAt(InArgs._WrapTextAt)
+		.LineBreakPolicy(InArgs._LineBreakPolicy)
+		.LineHeightPercentage(InArgs._LineHeightPercentage)
+		.Justification(InArgs._Justification)
+		.MinDesiredWidth(InArgs._MinDesiredWidth)
+		.TextShapingMethod(InArgs._TextShapingMethod)
+		.TextFlowDirection(InArgs._TextFlowDirection)
+		.SimpleTextMode(InArgs._SimpleTextMode)
+		);
+
+	SetRenderTransformPivot(FVector2d(0,0));
+	const FSlateRenderTransform Rotate90(FQuat2D(FMath::DegreesToRadians(90.f)));
+	SetRenderTransform(Rotate90);
+	
+}
+
+FVector2D SVerticalTextBlock::ComputeDesiredSize(const float X) const
+{
+	FVector2D WantedSize=STextBlock::ComputeDesiredSize(X);
+	double Tmp=WantedSize.X;
+	WantedSize.X=WantedSize.Y;
+	WantedSize.Y=Tmp;
+	return WantedSize;
+}
diff --git a/Engine/Plugins/Experimental/UserToolBoxCore/Source/UserToolBoxCore/Private/UTBBaseCommand.cpp b/Engine/Plugins/Experimental/UserToolBoxCore/Source/UserToolBoxCore/Private/UTBBaseCommand.cpp
new file mode 100644
--- /dev/null
+++ b/Engine/Plugins/Experimental/UserToolBoxCore/Source/UserToolBoxCore/Private/UTBBaseCommand.cpp
@@ -0,0 +1,119 @@
+// Copyright Epic Games, Inc. All Rights Reserved.
+
+#include "UTBBaseCommand.h"
+#include "Widgets/Input/SSpinBox.h"
+#include "Widgets/Input/SNumericEntryBox.h"
+#include "Dialog/SCustomDialog.h"
+#include "PropertyEditorModule.h"
+#include "IDetailsView.h"
+#include "UTBBaseTab.h"
+#include "Widgets/Layout/SBox.h"
+#include "ScopedTransaction.h"
+#include "Modules/ModuleManager.h"
+
+void UUTBBaseCommand::Execute()
+{
+}
+
+void UUTBBaseCommand::ExecuteCommand()
+{
+	
+	if (bShowParameters)
+	{
+		if (!DisplayParameters())
+		{
+			return;
+		}
+			
+	}
+	if (bIsTransaction)
+	{
+		const FScopedTransaction Transaction =FScopedTransaction(nullptr,FText::FromString(Name),nullptr);
+		Execute();	
+	}
+	else
+	{
+		Execute();
+	}
+	
+}
+
+bool UUTBBaseCommand::DisplayParameters()
+{
+	FPropertyEditorModule& PropertyEditorModule = FModuleManager::GetModuleChecked<FPropertyEditorModule>("PropertyEditor");
+	FDetailsViewArgs Args;
+	Args.NameAreaSettings = FDetailsViewArgs::HideNameArea;
+	Args.bHideSelectionTip = true;
+	Args.DefaultsOnlyVisibility = EEditDefaultsOnlyNodeVisibility::Hide;
+	Args.bAllowSearch = false;
+
+	
+	TSharedRef<IDetailsView> DetailsView=PropertyEditorModule.CreateDetailView(Args);
+	
+	DetailsView->SetIsPropertyVisibleDelegate(FIsPropertyVisible::CreateLambda([](const FPropertyAndParent& PropertyAndParent)
+	{
+		FProperty * Property=const_cast<FProperty*>(&PropertyAndParent.Property);
+		return !StaticClass()->HasProperty(Property);
+
+	}));
+	DetailsView->SetObject(this);
+	TSharedRef<SCustomDialog> OptionsWindow =
+	SNew(SCustomDialog)
+	.Title(FText::FromString(Name+" properties"))
+	.Content()
+	[
+		SNew(SBox)
+		.MinDesiredWidth(400)
+		[
+			DetailsView
+
+		]
+		
+	]
+	.ContentAreaPadding(FMargin(0.f, 0.f, 0.f, 0.f))
+	.HAlignContent(EHorizontalAlignment::HAlign_Left)
+	.VAlignContent(EVerticalAlignment::VAlign_Top)
+	.Buttons
+	(
+		{
+		SCustomDialog::FButton(FText::FromString("Ok")),
+		SCustomDialog::FButton(FText::FromString("Cancel")),
+		}
+	);
+	return OptionsWindow->ShowModal()==0;
+	
+}
+
+UUTBBaseCommand* UUTBBaseCommand::CopyCommand( UObject* Owner)const 
+{
+	UUTBTabSection* NewSection=Cast<UUTBTabSection>(Owner);
+	if (IsValid(NewSection))
+	{
+		UUTBBaseCommand* NewCommand=DuplicateObject(this,Owner);
+		NewSection->Commands.Add(NewCommand);
+		return NewCommand;	
+	}
+	return nullptr;
+}
+
+void UUTBBaseCommand::AddObjectToTransaction(UObject* Object)
+{
+	if (IsValid(Object))
+	{
+		Object->Modify();
+	}
+	
+}
+
+void UUTBBaseCommand::AddObjectsToTransaction(TArray<UObject*> Objects)
+{
+	for (UObject* Object:Objects)
+	{
+		if (IsValid(Object))
+		{
+			Object->Modify();
+		}
+	}
+}
+
+
diff --git a/Engine/Plugins/Experimental/UserToolBoxCore/Source/UserToolBoxCore/Private/UTBBaseTab.cpp b/Engine/Plugins/Experimental/UserToolBoxCore/Source/UserToolBoxCore/Private/UTBBaseTab.cpp
new file mode 100644
--- /dev/null
+++ b/Engine/Plugins/Experimental/UserToolBoxCore/Source/UserToolBoxCore/Private/UTBBaseTab.cpp
@@ -0,0 +1,394 @@
+// Copyright Epic Games, Inc. All Rights Reserved.
+
+
+#include "UTBBaseTab.h"
+
+#include "EditorUtilityBlueprint.h"
+#include "UICommandsScriptingSubsystem.h"
+#include "UserToolBoxBaseBlueprint.h"
+#include "UserToolBoxSubsystem.h"
+#include "UTBBaseUICommand.h"
+#include "UTBBaseUITab.h"
+#include "Engine/Engine.h"
+#include "Editor.h"
+
+const FName UUserToolBoxBaseTab::PlaceHolderSectionName="Hidden";
+
+
+TSharedPtr<SWidget> UUserToolBoxBaseTab::GenerateUI()
+{
+	return GenerateUI(nullptr,FUITemplateParameters());
+}
+TSharedPtr<SWidget>  UUserToolBoxBaseTab::GenerateUI(const TSubclassOf<UUTBDefaultUITemplate> TabUiClass,const FUITemplateParameters& Params)
+{
+	RemoveInvalidCommand();
+	TSubclassOf<UUTBDefaultUITemplate> RealTabUiClass=TabUiClass;
+	if (RealTabUiClass==nullptr)
+	{
+		RealTabUiClass= this->TabUI;
+	}
+	if (RealTabUiClass==nullptr)
+	{
+		RealTabUiClass=UUTBDefaultUITemplate::StaticClass();
+	}
+	if (IsValid(RealTabUiClass))
+	{
+		UUTBDefaultUITemplate* CurTabUI;
+		CurTabUI= NewObject<UUTBDefaultUITemplate>(GetTransientPackage(),RealTabUiClass);
+		TSharedPtr<SWidget> Widget=CurTabUI->BuildTabUI(this,Params);
+		return Widget;
+	}
+	return TSharedPtr<SWidget>();
+}
+void UUserToolBoxBaseTab::PostLoad()
+{
+	Super::PostLoad();
+	OnObjectReplacedHandle=FCoreUObjectDelegates::OnObjectsReplaced.AddUObject(this,&UUserToolBoxBaseTab::ReplaceCommands);
+}
+
+void UUserToolBoxBaseTab::BeginDestroy()
+{
+	UnregisterCommand();
+	if (OnObjectReplacedHandle.IsValid())
+	{
+		FCoreUObjectDelegates::OnObjectsReplaced.Remove(OnObjectReplacedHandle);	
+	}
+	Super::BeginDestroy();
+}
+
+void UUserToolBoxBaseTab::InsertCommandFromClass(const TSubclassOf<UUTBBaseCommand> InClass,const FString Section,const int Index)
+{
+	if (InClass==nullptr)
+	{
+		return;
+	}
+	InsertCommand(Cast<UUTBBaseCommand>(InClass->GetDefaultObject()),Section,Index);
+}
+
+void UUserToolBoxBaseTab::InsertCommand(const UUTBBaseCommand* InCommand,const FString SectionName,const int Index)
+{
+	UUTBTabSection* CurrentSection=GetOrCreateSection(SectionName);
+	UUTBBaseCommand* NewCommand=DuplicateObject(InCommand,CurrentSection);
+	if (Index<0 || Index>=CurrentSection->Commands.Num())
+	{
+		CurrentSection->Commands.Add(NewCommand);
+	}
+	else
+	{
+		CurrentSection->Commands.Insert(NewCommand,Index);
+	}
+}
+
+
+void UUserToolBoxBaseTab::RemoveCommand(const FString SectionName,const int Index)
+{
+	UUTBTabSection* CurrentSection=GetSection(SectionName);
+	if (CurrentSection!=nullptr)
+	{
+		CurrentSection->Commands.RemoveAt(Index);
+	}
+}
+
+
+
+UUTBTabSection* UUserToolBoxBaseTab::GetSection(const FString SectionName)
+{
+	Sections.Remove(nullptr);
+	if (SectionName==PlaceHolderSectionName.ToString())
+	{
+		return PlaceholderSection;
+	}
+	for (UUTBTabSection* Section:Sections)
+	{
+		if (Section->SectionName==SectionName)
+		{
+			return Section;
+		}
+	}
+	return nullptr;
+}
+
+UUTBTabSection* UUserToolBoxBaseTab::GetOrCreateSection( const FString SectionName)
+{
+	if (SectionName==PlaceHolderSectionName.ToString())
+	{
+		return GetPlaceHolderSection();
+	}
+	UUTBTabSection* CurrentSection=GetSection(SectionName);
+	if (IsValid(CurrentSection))
+	{
+		return CurrentSection;
+	}
+	CurrentSection=NewObject<UUTBTabSection>(this);
+	CurrentSection->SectionName=SectionName;
+	Sections.Add(CurrentSection);
+	return CurrentSection;
+}
+
+void UUserToolBoxBaseTab::RemoveSection( UUTBTabSection* Section)
+{
+	if (Section==PlaceholderSection)
+	{
+		return;
+	}
+	Sections.Remove(Section);
+}
+
+void UUserToolBoxBaseTab::RemoveSection(const FString SectionName)
+{
+	UUTBTabSection* Section=GetSection(SectionName);
+	if (IsValid(Section))
+	{
+		RemoveSection(Section);
+	}
+}
+
+void UUserToolBoxBaseTab::RemoveCommand(UUTBBaseCommand* Command)
+{
+	for (UUTBTabSection* Section:Sections)
+	{
+		Section->Commands.Remove(Command);
+	}
+	if (IsValid(PlaceholderSection))
+	{
+		PlaceholderSection->Commands.Remove(Command);
+	}
+}
+
+void UUserToolBoxBaseTab::RemoveCommandFromSection(UUTBBaseCommand* Command,const FString SectionName)
+{
+	UUTBTabSection* Section=GetSection(SectionName);
+	if (IsValid(Section))
+	{
+		Section->Commands.Remove(Command);
+	}
+}
+
+void UUserToolBoxBaseTab::ReplaceCommands(const TMap<UObject*, UObject*>& ReplacementMap)
+{
+	Sections.Remove(nullptr);
+	for (UUTBTabSection* Section:Sections)
+	{
+		for (UUTBBaseCommand* Command:Section->Commands)
+		{
+			UObject*const* Result=ReplacementMap.Find(Command);
+			if (Result!=nullptr)
+			{
+				Section->Commands[Section->Commands.Find(Command)]=Cast<UUTBBaseCommand>(*Result);	
+			}
+		}
+	}
+	if (IsValid(PlaceholderSection))
+	{
+		for (UUTBBaseCommand* Command:PlaceholderSection->Commands)
+		{
+			UObject*const* Result=ReplacementMap.Find(Command);
+			if (Result!=nullptr)
+			{
+				PlaceholderSection->Commands[PlaceholderSection->Commands.Find(Command)]=Cast<UUTBBaseCommand>(*Result);	
+			}
+		}	
+	}
+}
+
+TArray<UUTBTabSection*> UUserToolBoxBaseTab::GetSections()
+{
+	return Sections;
+}
+
+void UUserToolBoxBaseTab::MoveSectionAfterExistingSection(const FString SectionToMoveName,const FString SectionToTargetName)
+{
+	
+	UUTBTabSection* SectionToMove=GetSection(SectionToMoveName);
+	UUTBTabSection* SectionToTarget=GetSection(SectionToTargetName);
+	if (!(IsValid(SectionToMove)&&IsValid(SectionToTarget))&&SectionToMove==SectionToTarget)
+	{
+		return;
+	}
+	int32 Index=Sections.Find(SectionToTarget);
+	Sections.Remove(SectionToMove);
+	Sections.Insert(SectionToMove,Index);
+}
+
+UUTBTabSection* UUserToolBoxBaseTab::GetPlaceHolderSection()
+{
+	if (!IsValid(PlaceholderSection))
+	{
+		PlaceholderSection=NewObject<UUTBTabSection>(this);
+		PlaceholderSection->SectionName=PlaceHolderSectionName.ToString();
+	}
+	return PlaceholderSection;
+}
+
+bool UUserToolBoxBaseTab::ContainsCommand(const UUTBBaseCommand* Command)const
+{
+	if (!IsValid(this))
+	{
+		return false;
+	}
+	for (UUTBTabSection* Section:Sections)
+	{
+		if (Section->Commands.Contains(Command))
+		{
+			return true;
+		}
+	}
+	return false;
+}
+
+void UUserToolBoxBaseTab::RegisterCommand()
+{
+	UnregisterCommand();
+	UUICommandsScriptingSubsystem* UICommandsScriptingSubsystem = GEngine->GetEngineSubsystem<UUICommandsScriptingSubsystem>();
+	if (!IsValid(UICommandsScriptingSubsystem))
+	{
+		return;
+	}
+	for (UUTBTabSection* Section:Sections)
+	{
+		FString SectionName=Section->SectionName;
+		FName NewSetName=FName(this->GetName());
+		if (!UICommandsScriptingSubsystem->IsCommandSetRegistered(NewSetName))
+		{
+			UICommandsScriptingSubsystem->RegisterCommandSet(NewSetName);	
+		}
+		
+		
+		for (UUTBBaseCommand* Command:Section->Commands)
+		{
+			if (!IsValid(Command))
+			{
+				continue;
+			}
+			if (Command->KeyboardShortcut.IsValidChord())
+			{
+				FScriptingCommandInfo ScriptingCommandInfo;
+				ScriptingCommandInfo.InputChord=Command->KeyboardShortcut;
+				ScriptingCommandInfo.Description=FText::FromString(Command->Tooltip);
+				ScriptingCommandInfo.Label=FText::FromString(Section->SectionName+Command->Name);
+				ScriptingCommandInfo.Set=NewSetName;
+				ScriptingCommandInfo.ContextName="LevelViewport";
+				ScriptingCommandInfo.Name=FName(Command->Name);
+				FExecuteCommand ExecuteCommand;
+				ExecuteCommand.BindUFunction(Command,FName("ExecuteCommand"));
+				UICommandsScriptingSubsystem->RegisterCommand(ScriptingCommandInfo,ExecuteCommand,true);
+				ScriptingCommandInfo.ContextName="ContentBrowser";
+				UICommandsScriptingSubsystem->RegisterCommand(ScriptingCommandInfo,ExecuteCommand,true);
+			}
+		}
+	}
+}
+
+void UUserToolBoxBaseTab::UnregisterCommand()
+{
+	UUICommandsScriptingSubsystem* UICommandsScriptingSubsystem = GEngine->GetEngineSubsystem<UUICommandsScriptingSubsystem>();
+	if (!IsValid(UICommandsScriptingSubsystem))
+	{
+		return;
+	}
+	if (UICommandsScriptingSubsystem->IsCommandSetRegistered(FName(GetName())))
+	{
+		UICommandsScriptingSubsystem->UnregisterCommandSet(FName(GetName()));	
+	}
+}
+
+void UUserToolBoxBaseTab::RemoveInvalidCommand()
+{
+	
+	for (UUTBTabSection* Section:Sections)
+	{
+		if (IsValid(Section))
+		{
+			Section->Commands.RemoveAll([](UUTBBaseCommand* Command)
+			{
+				return !IsValid(Command);
+			});
+		}
+	}
+}
+
+
+SUserToolBoxTabWidget::~SUserToolBoxTabWidget()
+{
+	if (OnTabChangedDelegate.IsValid())
+	{
+		UUserToolboxSubsystem* UTBSubsystem=GEditor->GetEditorSubsystem<UUserToolboxSubsystem>();
+		if (IsValid(UTBSubsystem))
+		{
+			UTBSubsystem->OnTabChanged.Remove(OnTabChangedDelegate);
+			OnTabChangedDelegate.Reset();	
+		}
+		
+		FCoreUObjectDelegates::OnObjectPropertyChanged.Remove(OnObjectPropertyChangedDelegate);
+	}
+	if (IsValid(Tab))
+	{
+		Tab->UnregisterCommand();
+	}
+}
+
+void SUserToolBoxTabWidget::Construct(const FArguments& InArgs)
+{
+	Tab=InArgs._Tab.Get();
+	UIOverride=InArgs._UIOverride.Get();
+	Parameters=InArgs._UIParameters.Get();
+	UpdateTab(Tab);
+	if (!OnTabChangedDelegate.IsValid())
+	{
+		UUserToolboxSubsystem* UTBSubsystem=GEditor->GetEditorSubsystem<UUserToolboxSubsystem>();
+		OnTabChangedDelegate=UTBSubsystem->OnTabChanged.AddSP(this,&SUserToolBoxTabWidget::UpdateTab);
+		OnObjectPropertyChangedDelegate=FCoreUObjectDelegates::OnObjectPropertyChanged.AddSP(this,&SUserToolBoxTabWidget::OnObjectModified);
+	}
+	Tab->RegisterCommand();
+}
+
+void SUserToolBoxTabWidget::UpdateTab(UUserToolBoxBaseTab* InTab)
+{
+	if (!IsValid(Tab) || Tab!=InTab )
+		return;
+	if (!IsValid(Tab))
+	{
+		ChildSlot
+		[
+			SNullWidget::NullWidget
+		];
+		return;
+	}
+	ChildSlot
+	[
+		Tab->GenerateUI(IsValid(UIOverride)?UIOverride:nullptr,Parameters).ToSharedRef()
+	];	
+	Tab->RegisterCommand();
+}
+
+void SUserToolBoxTabWidget::OnObjectModified(UObject* Object, FPropertyChangedEvent& PropertyChangedEventEvent)
+{
+	/*
+	 * We are looking for each modified object to check update of a command and update the UI
+	 */
+	if (Object->IsA<UUTBBaseCommand>())
+	{
+		if (!IsValid(Tab))
+		{
+			return;
+		}
+		if (Tab->ContainsCommand(Cast<UUTBBaseCommand>(Object)))
+		{
+			UpdateTab(Tab);
+		}
+	}
+	if (Object->IsA<UUserToolBoxBaseTab>())
+	{
+		if (!IsValid(Tab))
+		{
+			return;
+		}
+		if (Tab==Object)
+		{
+			if (PropertyChangedEventEvent.Property->GetName()=="TabUI")
+			{
+				UpdateTab(Tab);
+			}
+		}
+	}
+}
diff --git a/Engine/Plugins/Experimental/UserToolBoxCore/Source/UserToolBoxCore/Private/UTBBaseUICommand.cpp b/Engine/Plugins/Experimental/UserToolBoxCore/Source/UserToolBoxCore/Private/UTBBaseUICommand.cpp
new file mode 100644
--- /dev/null
+++ b/Engine/Plugins/Experimental/UserToolBoxCore/Source/UserToolBoxCore/Private/UTBBaseUICommand.cpp
@@ -0,0 +1,170 @@
+// Copyright Epic Games, Inc. All Rights Reserved.
+
+
+#include "UTBBaseUICommand.h"
+#include "UTBBaseCommand.h"
+#include "UserToolBoxStyle.h"
+#include "ClassViewerFilter.h"
+#include "Kismet2/KismetEditorUtilities.h"
+#include "Components/CanvasPanel.h"
+#include "Widgets/SBoxPanel.h"
+#include "Widgets/Input/SButton.h"
+#include "Widgets/Text/STextBlock.h"
+#include "Widgets/Images/SImage.h"
+
+bool UBaseCommandNativeUI::IsSupportingCommandClass(TSubclassOf<UUTBBaseCommand> CommandClass)
+{
+
+	return true;
+}
+
+void UBaseCommandNativeUI::SetCommand(UUTBBaseCommand* Command)
+{
+	MyCommand=Command;
+	if (MyCommand->GetClass()->ClassGeneratedBy!=nullptr)
+	{
+		FCoreUObjectDelegates::OnObjectsReplaced.AddUObject(this,&UBaseCommandNativeUI::ReplaceCommand);
+	}
+
+}
+
+TSharedRef<SWidget> UBaseCommandNativeUI::GetUI()
+{
+
+	TSharedPtr<SWidget> Widget;
+	TSharedPtr<SVerticalBox> VerticalBox;
+	
+	SAssignNew(Widget,SVerticalBox)
+	+SVerticalBox::Slot()
+	.HAlign(HAlign_Fill)
+	.VAlign(VAlign_Fill)
+	[
+
+			SNew(SButton)
+			.ButtonColorAndOpacity(FLinearColor::Gray)
+			.HAlign(HAlign_Fill)
+			.VAlign(VAlign_Center)
+				
+			.OnClicked_Lambda([this]
+				()
+			{
+				check(MyCommand)
+				
+				MyCommand->ExecuteCommand();
+				return FReply::Handled();
+			})
+			[
+
+				SAssignNew(VerticalBox, SVerticalBox)
+
+				+ SVerticalBox::Slot()
+				.AutoHeight()
+				.HAlign(HAlign_Center)
+				
+				[
+					SNew(STextBlock)
+					.Justification(ETextJustify::Center)
+					.AutoWrapText(true)
+					.ColorAndOpacity(FLinearColor::White)
+					.Text(FText::FromString(MyCommand->Name))
+				]
+			]
+	
+		
+	];
+	
+	if (!MyCommand->IconPath.IsEmpty())
+	{
+		VerticalBox->InsertSlot(0)
+			.AutoHeight()
+			.HAlign(HAlign_Center)
+			.Padding(1)
+			[
+				SNew(SImage)
+				.Image(FUserToolBoxStyle::Get().GetBrush(FName(MyCommand->IconPath)))
+			];
+	}
+	return Widget.ToSharedRef();
+	
+	
+}
+
+void UBaseCommandNativeUI::ReplaceCommand(const TMap<UObject*, UObject*>& Map)
+{
+	if (Map.Find(MyCommand)!=nullptr)
+	{
+		MyCommand=Cast<UUTBBaseCommand>(*Map.Find(MyCommand));
+	}
+		
+}
+
+
+
+void UBaseCommandNativeUI::BeginDestroy()
+{
+	FCoreUObjectDelegates::OnObjectsReplaced.RemoveAll(this);
+	UObject::BeginDestroy();
+}
+
+// Add default functionality here for any IUTBBaseUICommand functions that are not pure virtual.
+void UBaseCommandNativeUI::ExecuteCurrentCommand()
+{
+	if (MyCommand!=nullptr)
+	{
+		MyCommand->ExecuteCommand();
+	}
+}
+
+
+
+
+
+
+void UUTBCommandUMGUI::ExecuteCommand()
+{
+	ExecuteCurrentCommand();
+}
+
+bool UUTBCommandUMGUI::IsSupportingCommandClass(TSubclassOf<UUTBBaseCommand> CommandClass)
+{
+	return DoesSupportClass(CommandClass);
+}
+
+void UUTBCommandUMGUI::SetCommand(UUTBBaseCommand* _Command)
+{
+	Command=_Command;
+	if (Command->GetClass()->ClassGeneratedBy!=nullptr)
+	{
+		FCoreUObjectDelegates::OnObjectsReplaced.AddUObject(this,&UUTBCommandUMGUI::ReplaceCommand);
+	}
+}
+
+void UUTBCommandUMGUI::ExecuteCurrentCommand()
+{
+	if (Command!=nullptr)
+	{
+		Command->ExecuteCommand();
+	}
+}
+
+TSharedRef<SWidget> UUTBCommandUMGUI::GetUI()
+{
+	return TakeWidget();
+}
+
+void UUTBCommandUMGUI::ReplaceCommand(const TMap<UObject*, UObject*>& Map)
+{
+	if (Map.Find(Command)!=nullptr)
+	{
+		Command=Cast<UUTBBaseCommand>(*Map.Find(Command));
+	}
+		
+}
+
+void UUTBCommandUMGUI::BeginDestroy()
+{
+	FCoreUObjectDelegates::OnObjectsReplaced.RemoveAll(this);
+
+	Super::BeginDestroy();
+}
+
diff --git a/Engine/Plugins/Experimental/UserToolBoxCore/Source/UserToolBoxCore/Private/UTBBaseUITab.cpp b/Engine/Plugins/Experimental/UserToolBoxCore/Source/UserToolBoxCore/Private/UTBBaseUITab.cpp
new file mode 100644
--- /dev/null
+++ b/Engine/Plugins/Experimental/UserToolBoxCore/Source/UserToolBoxCore/Private/UTBBaseUITab.cpp
@@ -0,0 +1,302 @@
+// Copyright Epic Games, Inc. All Rights Reserved.
+
+
+#include "UTBBaseUITab.h"
+#include "Widgets/Layout/SScrollBox.h"
+#include "UserToolBoxStyle.h"
+#include "Widgets/Layout/SExpandableArea.h"
+#include "Dialog/SCustomDialog.h"
+#include "UTBBaseTab.h"
+#include "UTBBaseCommand.h"
+#include "UTBBaseUICommand.h"
+#include "Widgets/Layout/SWrapBox.h"
+#include "Widgets/Input/SButton.h"
+#include "Widgets/Images/SImage.h"
+#include "Editor.h"
+#include "Subsystems/AssetEditorSubsystem.h"
+#include "Framework/Commands/InputBindingManager.h"
+
+TSharedPtr<SWidget> UUTBDefaultUITemplate::BuildTabUI(UUserToolBoxBaseTab* Tab, const FUITemplateParameters& Params)
+{
+	if (Tab==nullptr)
+	{
+		return TSharedPtr<SWidget>();
+	}
+	TSharedPtr<SScrollBox> ScrollBox;
+	SAssignNew(ScrollBox,SScrollBox)
+	.Orientation(Orient_Vertical);
+	if (!Params.bForceHideSettingsButton)
+	{
+		ScrollBox->AddSlot()
+		.HAlign(HAlign_Right)
+		.VAlign(VAlign_Bottom)
+		[
+			SNew(SButton)
+			.ButtonStyle(&FAppStyle::Get().GetWidgetStyle<FButtonStyle>("SimpleButton"))
+			.Visibility_Lambda([Tab]()
+			{
+				return Tab->IsSettingShouldBeVisible?EVisibility::Visible:EVisibility::Collapsed;
+			})
+			[
+				SNew(SImage)
+				.ColorAndOpacity(FSlateColor::UseForeground())
+				.Image(FAppStyle::Get().GetBrush("Icons.Adjust"))
+			]
+			.OnClicked_Lambda([Tab]()
+			{
+			
+				UAssetEditorSubsystem* AssetEditorSubsystem=GEditor->GetEditorSubsystem<UAssetEditorSubsystem>();
+				if (IsValid(AssetEditorSubsystem) )
+				{
+					AssetEditorSubsystem->OpenEditorForAsset(Tab);
+				}
+
+				return FReply::Handled();
+			})
+		];	
+	}
+	
+	
+	for (UUTBTabSection* Section:Tab->GetSections())
+	{
+		if (!IsValid(Section))
+		{
+			continue;
+		}
+		TSharedPtr<SWrapBox> WrapBox;
+		FString SectionName=Section->SectionName;
+		if (Params.bPostfixSectionWithTabName)
+		{
+			SectionName+= " - "+Tab->Name;
+		}
+		if (Params.bPrefixSectionWithTabName)
+		{
+			SectionName= Tab->Name+" - "+SectionName;
+		}
+		
+		ScrollBox->AddSlot()
+		[
+			SNew(SExpandableArea)
+			.AreaTitle(FText::FromString(SectionName))
+			.BorderBackgroundColor(FLinearColor(.6, .6, .6, 1.0f))
+			.BodyBorderBackgroundColor(FLinearColor::Transparent)
+			.BodyContent()
+			[
+				SAssignNew(WrapBox,SWrapBox)
+				.UseAllottedSize(true)
+				
+			]
+		];
+		TMap<int,TSet<int>> SlotToSkip;
+		for (UUTBBaseCommand* Command:Section->Commands)
+		{
+			if (!IsValid(Command))
+			{
+				continue;
+			}
+			UClass * UIToUse=UBaseCommandNativeUI::StaticClass();
+			// checking the override
+			// least priority first
+			if (Tab->DefaultCommandUIOverride!=nullptr)
+			{
+				IUTBUICommand* UIInterface=Cast<IUTBUICommand>(Tab->DefaultCommandUIOverride->GetDefaultObject());
+				if (UIInterface!=nullptr)
+				{
+					if (UIInterface->IsSupportingCommandClass(Command->GetClass()))
+					{
+						UIToUse=Tab->DefaultCommandUIOverride;
+					}
+				}
+			}
+			if (Command->UIOverride!=nullptr)
+			{
+				IUTBUICommand* UIInterface=Cast<IUTBUICommand>(Command->UIOverride->GetDefaultObject());
+				if (UIInterface!=nullptr)
+				{
+					if (UIInterface->IsSupportingCommandClass(Command->GetClass()))
+					{
+						UIToUse=Command->UIOverride;
+					}
+				}
+			}
+			Command->UI=NewObject<UObject>(GetTransientPackage(),UIToUse);
+			Cast<IUTBUICommand>(Command->UI)->SetCommand(Command);
+			WrapBox->AddSlot()
+			.FillEmptySpace(false)
+			[
+				SNew(SBox)
+				.WidthOverride(120)
+				.HeightOverride(70)
+				[
+				Cast<IUTBUICommand>(Command->UI)->GetUI()
+				]	
+			];
+		}
+	}
+	return ScrollBox;
+}
+
+
+TSharedPtr<SWidget> UUTBToolBarTabUI::BuildTabUI(UUserToolBoxBaseTab* Tab, const FUITemplateParameters& Params)
+{
+	TSharedRef<SVerticalBox> ToolBoxVBox = SNew(SVerticalBox);
+	if (!Params.bForceHideSettingsButton)
+	{
+		ToolBoxVBox->AddSlot()
+		.VAlign(EVerticalAlignment::VAlign_Top)
+		.HAlign(EHorizontalAlignment::HAlign_Right)
+		.AutoHeight()
+		.Padding(0)
+		[
+			SNew(SButton)
+			.ButtonStyle(&FAppStyle::Get().GetWidgetStyle<FButtonStyle>("SimpleButton"))
+			.Visibility_Lambda([Tab]()
+			{
+				return Tab->IsSettingShouldBeVisible?EVisibility::Visible:EVisibility::Collapsed;
+			})
+			[
+				SNew(SImage)
+				.ColorAndOpacity(FSlateColor::UseForeground())
+				.Image(FAppStyle::Get().GetBrush("Icons.Adjust"))
+			]
+			.OnClicked_Lambda([Tab]()
+			{
+				UAssetEditorSubsystem* AssetEditorSubsystem=GEditor->GetEditorSubsystem<UAssetEditorSubsystem>();
+				if (IsValid(AssetEditorSubsystem) )
+				{
+					AssetEditorSubsystem->OpenEditorForAsset(Tab);
+				}
+				return FReply::Handled();
+			})
+		];	
+	}
+	
+	TArray<UUTBTabSection*> Sections=Tab->GetSections();
+	TMap<FString,TArray<TSharedPtr<FUICommandInfo>>> UICommandInfosPerSection;
+	TSharedPtr<class FUICommandList> TabCommands = MakeShareable(new FUICommandList);
+	for (UUTBTabSection* Section:Sections)
+	{
+		TSharedPtr<FBindingContext>   BindingContext= MakeShared<FBindingContext>(FName(FString("UTB.")+Tab->GetName()+"."+Section->GetName()),FText::FromString(FString("UTB.")+Tab->GetName()+"."+Section->GetName()),NAME_None,FAppStyle::Get().GetStyleSetName());
+		TSharedPtr<FToolBarBuilder> ToolbarBuilder=GetBuilder(TabCommands, FMultiBoxCustomization(FName(Tab->Name)));
+		ToolbarBuilder->SetStyle(&FAppStyle::Get(), "PaletteToolBar");
+		FSlateIcon EmptyIcon(FAppStyle::Get().GetStyleSetName(),"TableView.Header.Column");
+		int CommandCounter=0;
+		for (UUTBBaseCommand* Command:Section->Commands)
+		{
+			FString CommandName=Command->Name+"_"+FString::FromInt(CommandCounter);
+			TSharedPtr<FUICommandInfo> CommandInfo=FInputBindingManager::Get().FindCommandInContext(BindingContext->GetContextName(),FName(CommandName));
+			if (CommandInfo.IsValid())
+			{
+				FInputBindingManager::Get().RemoveInputCommand(BindingContext.ToSharedRef(),CommandInfo.ToSharedRef());
+			}
+			FUICommandInfo::MakeCommandInfo(
+				BindingContext.ToSharedRef(),
+				CommandInfo,FName(CommandName),
+				FText::FromString(Command->Name),
+				FText::FromString(Command->Tooltip),
+				Command->IconPath.IsEmpty()?EmptyIcon:
+				FSlateIcon(FUserToolBoxStyle::GetStyleSetName(), FName(Command->IconPath)),
+				EUserInterfaceActionType::Button,FInputChord(),FInputChord(),NAME_None
+				);	
+			TabCommands->MapAction(CommandInfo,FExecuteAction::CreateUObject(Command,&UUTBBaseCommand::ExecuteCommand));
+			UClass * UIToUse=UBaseCommandNativeUI::StaticClass();
+			if (Tab->DefaultCommandUIOverride!=nullptr)
+			{
+				IUTBUICommand* UIInterface=Cast<IUTBUICommand>(Tab->DefaultCommandUIOverride->GetDefaultObject());
+				if (UIInterface!=nullptr)
+				{
+					if (UIInterface->IsSupportingCommandClass(Command->GetClass()))
+					{
+						UIToUse=Tab->DefaultCommandUIOverride;
+					}
+				}
+			}
+			if (Command->UIOverride!=nullptr)
+			{
+				IUTBUICommand* UIInterface=Cast<IUTBUICommand>(Command->UIOverride->GetDefaultObject());
+				if (UIInterface!=nullptr)
+				{
+					if (UIInterface->IsSupportingCommandClass(Command->GetClass()))
+					{
+						UIToUse=Command->UIOverride;
+					}
+				}
+			}
+			Command->UI=NewObject<UObject>(GetTransientPackage(),UIToUse);
+			Cast<IUTBUICommand>(Command->UI)->SetCommand(Command);
+			if (UIToUse!=UBaseCommandNativeUI::StaticClass())
+			{
+				ToolbarBuilder->AddToolBarWidget(Cast<IUTBUICommand>(Command->UI)->GetUI());	
+			}
+			else
+			{
+				ToolbarBuilder->AddToolBarButton(CommandInfo);
+			}
+			ToolbarBuilder->AddSeparator();
+			CommandCounter++;
+		}
+		FString SectionName=Section->SectionName;
+		if (Params.bPostfixSectionWithTabName)
+		{
+			SectionName+= " - "+Tab->Name;
+		}
+		if (Params.bPrefixSectionWithTabName)
+		{
+			SectionName= Tab->Name+" - "+SectionName;
+		}
+		
+		ToolBoxVBox->AddSlot()
+		.AutoHeight()
+		.VAlign(EVerticalAlignment::VAlign_Top)
+		.Padding(FMargin(0.0, 0.0))
+		[
+			SNew(SExpandableArea)
+			.AreaTitle(FText::FromString(SectionName))
+			.AreaTitleFont(FAppStyle::Get().GetFontStyle("NormalFont"))
+			.BorderImage(FAppStyle::Get().GetBrush("PaletteToolbar.ExpandableAreaHeader"))
+			.BodyBorderImage(FAppStyle::Get().GetBrush("PaletteToolbar.ExpandableAreaBody"))
+			.HeaderPadding(FMargin(4.f))
+			.InitiallyCollapsed(true)
+			.Padding(FMargin(4.0, 0.0))
+			.BodyContent()
+			[
+				SNew(SBox)
+				[
+					ToolbarBuilder->MakeWidget()
+				]	
+			]
+		];
+	}
+	ToolBoxVBox->AddSlot()
+	.AutoHeight()
+	.VAlign(EVerticalAlignment::VAlign_Fill)
+	[
+		SNew(SVerticalBox)
+	];
+	return ToolBoxVBox;
+}
+
+
+TSharedPtr<FToolBarBuilder> UUTBToolBarTabUI::GetBuilder(TSharedPtr<const FUICommandList> InCommandList,
+	FMultiBoxCustomization InCustomization, TSharedPtr<FExtender> InExtender, const bool InForceSmallIcons)
+{
+	return MakeShareable(new FToolBarBuilder(InCommandList,InCustomization,InExtender,InForceSmallIcons));
+}
+
+TSharedPtr<FToolBarBuilder> UUTBPaletteTabUI::GetBuilder(TSharedPtr<const FUICommandList> InCommandList,
+	FMultiBoxCustomization InCustomization, TSharedPtr<FExtender> InExtender, const bool InForceSmallIcons)
+{
+	return MakeShareable(new FUniformToolBarBuilder(InCommandList,InCustomization,InExtender,InForceSmallIcons));
+}
+
+TSharedPtr<FToolBarBuilder> UUTBVerticalToolBarTabUI::GetBuilder(TSharedPtr<const FUICommandList> InCommandList,
+	FMultiBoxCustomization InCustomization, TSharedPtr<FExtender> InExtender, const bool InForceSmallIcons)
+{
+	return MakeShareable(new FVerticalToolBarBuilder(InCommandList,InCustomization,InExtender,InForceSmallIcons));
+}
+
+TSharedPtr<FToolBarBuilder> UUTBSlimHorizontalToolBarTabUI::GetBuilder(TSharedPtr<const FUICommandList> InCommandList,
+	FMultiBoxCustomization InCustomization, TSharedPtr<FExtender> InExtender, const bool InForceSmallIcons)
+{
+	return MakeShareable(new FSlimHorizontalToolBarBuilder(InCommandList,InCustomization,InExtender,InForceSmallIcons));
+}
diff --git a/Engine/Plugins/Experimental/UserToolBoxCore/Source/UserToolBoxCore/Private/UTBEditorBlueprint.cpp b/Engine/Plugins/Experimental/UserToolBoxCore/Source/UserToolBoxCore/Private/UTBEditorBlueprint.cpp
new file mode 100644
--- /dev/null
+++ b/Engine/Plugins/Experimental/UserToolBoxCore/Source/UserToolBoxCore/Private/UTBEditorBlueprint.cpp
@@ -0,0 +1,5 @@
+// Copyright Epic Games, Inc. All Rights Reserved.
+
+
+#include "UTBEditorBlueprint.h"
+
diff --git a/Engine/Plugins/Experimental/UserToolBoxCore/Source/UserToolBoxCore/Private/UTBFactory.cpp b/Engine/Plugins/Experimental/UserToolBoxCore/Source/UserToolBoxCore/Private/UTBFactory.cpp
new file mode 100644
--- /dev/null
+++ b/Engine/Plugins/Experimental/UserToolBoxCore/Source/UserToolBoxCore/Private/UTBFactory.cpp
@@ -0,0 +1,260 @@
+// Copyright Epic Games, Inc. All Rights Reserved.
+
+
+#include "UTBFactory.h"
+#include "Kismet2/KismetEditorUtilities.h"
+#include "UserToolBoxBaseBlueprint.h"
+#include "UTBEditorBlueprint.h"
+#include "Kismet2/SClassPickerDialog.h"
+#include "ClassViewerFilter.h"
+#include "IconsTracker.h"
+#include "UTBBaseTab.h"
+#include "Widgets/Layout/SScrollBox.h"
+#include "UObject/UObjectIterator.h"
+#include "Widgets/SWindow.h"
+#include "Widgets/Views/STreeView.h"
+#include "Widgets/Text/STextBlock.h"
+#include "Editor.h"
+
+UUTBEditorUtilityBlueprintFactory::UUTBEditorUtilityBlueprintFactory(const FObjectInitializer& ObjectInitializer)
+	: Super(ObjectInitializer)
+{
+	bCreateNew = true;
+	bEditAfterNew = true;
+	SupportedClass = UUTBEditorBlueprint::StaticClass();
+}
+UObject* UUTBEditorUtilityBlueprintFactory::FactoryCreateNew(UClass* InClass, UObject* InParent, FName InName,
+	EObjectFlags Flags, UObject* Context, FFeedbackContext* Warn)
+{
+	check(InClass->IsChildOf(UBlueprint::StaticClass()));
+
+
+		return FKismetEditorUtilities::CreateBlueprint(UUserToolBoxBaseBlueprint::StaticClass(), InParent, InName, BPTYPE_Normal, UUTBEditorBlueprint::StaticClass(), UBlueprintGeneratedClass::StaticClass());
+
+}
+
+bool UUTBEditorUtilityBlueprintFactory::CanCreateNew() const
+{
+	return true;
+}
+
+
+
+class FUTBCommandFactoryFilter : public IClassViewerFilter
+{
+public:
+	TSet< const UClass* > AllowedChildOfClasses;
+
+	TSet< const UClass*> DisallowedChildOfClasses;
+
+	bool IsClassAllowed(const FClassViewerInitializationOptions& InInitOptions, const UClass* InClass, TSharedRef< FClassViewerFilterFuncs > InFilterFuncs ) override
+	{
+		if (DisallowedChildOfClasses.Num() == 0 && AllowedChildOfClasses.Num() == 0)
+		{
+			return true;
+		}
+
+	
+		bool IsOk=InFilterFuncs->IfInChildOfClassesSet(AllowedChildOfClasses, InClass) != EFilterReturn::Failed;
+		IsOk=IsOk && InFilterFuncs->IfInChildOfClassesSet(DisallowedChildOfClasses, InClass) == EFilterReturn::Failed; 
+		return IsOk; 
+			
+	}
+
+	virtual bool IsUnloadedClassAllowed(const FClassViewerInitializationOptions& InInitOptions, const TSharedRef< const IUnloadedBlueprintData > InUnloadedClassData, TSharedRef< FClassViewerFilterFuncs > InFilterFuncs) override
+	{
+		if (DisallowedChildOfClasses.Num() == 0 && AllowedChildOfClasses.Num() == 0)
+		{
+			return true;
+		}
+
+		return InFilterFuncs->IfInChildOfClassesSet(AllowedChildOfClasses, InUnloadedClassData) != EFilterReturn::Failed
+			&& (InFilterFuncs->IfInChildOfClassesSet(DisallowedChildOfClasses, InUnloadedClassData) == EFilterReturn::Failed);;
+	}
+};
+
+
+
+UUTBCommandFactory::UUTBCommandFactory(const FObjectInitializer& ObjectInitializer)
+	: Super(ObjectInitializer)
+{
+	bCreateNew = true;
+	bEditAfterNew = true;
+	SupportedClass = UUTBBaseCommand::StaticClass();
+	ParentClass= nullptr;
+}
+UObject* UUTBCommandFactory::FactoryCreateNew(UClass* InClass, UObject* InParent, FName InName, EObjectFlags Flags,
+	UObject* Context, FFeedbackContext* Warn)
+{
+	UUTBBaseCommand* Command = NewObject<UUTBBaseCommand>(InParent, ParentClass, InName, Flags);
+	return Command;
+}
+
+bool UUTBCommandFactory::CanCreateNew() const
+{
+	return true;
+}
+
+bool UUTBCommandFactory::ConfigureProperties()
+{
+	TArray<UClass*> ValidClasses;
+	for (TObjectIterator<UClass> ClassIterator; ClassIterator; ++ClassIterator)
+	{
+		// Skip SKEL and REINST classes.
+		if (ClassIterator->GetName().StartsWith(TEXT("SKEL_")) || ClassIterator->GetName().StartsWith(TEXT("REINST_")))
+		{
+			continue;
+		}
+		if (ClassIterator->IsChildOf(UUTBBaseCommand::StaticClass()) && *ClassIterator != UUTBBaseCommand::StaticClass() && *ClassIterator != UUserToolBoxBaseBlueprint::StaticClass())
+		{
+				ValidClasses.Add(*ClassIterator);
+		}
+	}
+	TMap<FString,TArray<UClass*>> CDOPerCategory;
+	for (UClass* ValidClass:ValidClasses)
+	{
+		UUTBBaseCommand* CDO =Cast<UUTBBaseCommand>(ValidClass->GetDefaultObject());
+		if (IsValid(CDO))
+		{
+			CDOPerCategory.FindOrAdd(CDO->Category).Add(ValidClass);
+		}
+	}
+	TSharedRef<SWindow> Window = SNew(SWindow)
+		.Title(FText::FromString("Pick a parent command"))
+		.ClientSize(FVector2D(200.f, 300.0f))
+		.IsTopmostWindow(true);
+
+	
+
+	struct TreeNode
+	{
+		bool bIsCategory= false;
+		UClass*		  CommandClass;
+		FString			  Name;
+		TArray<TSharedPtr<TreeNode>> Children;
+	};
+	TArray<TSharedPtr<TreeNode>> Tree;
+	for (TTuple<FString,TArray<UClass*>> Category:CDOPerCategory)
+	{
+		FString CategoryFullName=Category.Key;
+		if (CategoryFullName.IsEmpty())
+		{
+			CategoryFullName="Uncategorized";
+		}
+		
+		TArray<FString> CategoriesSplit;
+		CategoryFullName.ParseIntoArray(CategoriesSplit,TEXT("|"),1);
+		TArray<TSharedPtr<TreeNode>>* CurrentArray=&Tree;
+		while (!CategoriesSplit.IsEmpty())
+		{
+			TSharedPtr<TreeNode>* CurrentNode= CurrentArray->FindByPredicate([&CategoriesSplit](TSharedPtr<TreeNode> Node)
+			{
+				return CategoriesSplit[0].Compare(Node->Name)==0;
+			});
+			if (CurrentNode==nullptr)
+			{
+				TSharedPtr<TreeNode> NewNode=CurrentArray->Add_GetRef(MakeShared<TreeNode>());
+				NewNode->bIsCategory=true;
+				NewNode->Name=CategoriesSplit[0];
+				CurrentNode=&NewNode;
+			}
+			
+			CurrentArray=&(*CurrentNode)->Children;
+			CategoriesSplit.RemoveAt(0);
+		}
+		for (UClass* ValidClass:Category.Value)
+		{
+			TSharedPtr<TreeNode> NewNode=CurrentArray->Add_GetRef(MakeShared<TreeNode>());
+			NewNode->Name=ValidClass->GetName();
+			NewNode->bIsCategory=false;
+			NewNode->CommandClass=ValidClass;
+		}
+	}
+	UClass* SelectedValue;
+	Window->SetContent(
+		SNew(SScrollBox)
+		+SScrollBox::Slot()
+		[
+			SNew(STreeView<TSharedPtr<TreeNode>>)
+			.ItemHeight(24)
+			.TreeItemsSource(&Tree)
+			.OnGenerateRow(STreeView<TSharedPtr<TreeNode>>::FOnGenerateRow::CreateLambda([](TSharedPtr<TreeNode> Node,const TSharedRef<STableViewBase>& OwnerTable)
+			{
+				TSharedPtr<STextBlock> TextBlock;
+				if (Node.IsValid())
+				{
+					SAssignNew(TextBlock,STextBlock).Text(FText::FromString(Node->Name));
+				}
+				else
+				{
+					SAssignNew(TextBlock,STextBlock).Text(FText::FromString("Invalid node"));
+				}
+				return SNew(STableRow<TSharedPtr<TreeNode>>,OwnerTable)
+				[
+					TextBlock.ToSharedRef()
+				];
+			}))
+			.OnGetChildren(STreeView<TSharedPtr<TreeNode>>::FOnGetChildren::CreateLambda([](TSharedPtr<TreeNode> ParentNode,TArray<TSharedPtr<TreeNode>>& Children)
+			{
+				Children=ParentNode->Children;
+			}))
+			.OnMouseButtonClick(STreeView<TSharedPtr<TreeNode>>::FOnMouseButtonClick::CreateLambda([&SelectedValue,&Window](TSharedPtr<TreeNode> Node)
+			{
+				if (!Node->bIsCategory)
+				{
+					SelectedValue=Node->CommandClass;
+					Window->RequestDestroyWindow();
+				}
+			}))
+		]
+		);
+	
+	GEditor->EditorAddModalWindow(Window);
+	if (IsValid(SelectedValue))
+	{
+		ParentClass=SelectedValue;
+	}
+	else
+	{
+		ParentClass=nullptr;
+	}
+	return ParentClass!=nullptr;
+}
+
+
+UUTBCommandTabFactory::UUTBCommandTabFactory(const FObjectInitializer& ObjectInitializer)
+	: Super(ObjectInitializer)
+{
+	bCreateNew = true;
+	bEditAfterNew = true;
+	SupportedClass = UUserToolBoxBaseTab::StaticClass();
+}
+UObject* UUTBCommandTabFactory::FactoryCreateNew(UClass* InClass, UObject* InParent, FName InName, EObjectFlags Flags,
+	UObject* Context, FFeedbackContext* Warn)
+{
+	UUserToolBoxBaseTab* Tab = NewObject<UUserToolBoxBaseTab>(InParent, UUserToolBoxBaseTab::StaticClass(), InName, Flags);
+	return Tab;
+}
+
+bool UUTBCommandTabFactory::CanCreateNew() const
+{
+	return true;
+}
+
+UIconTrackerFactory::UIconTrackerFactory(const FObjectInitializer& ObjectInitializer)
+	: Super(ObjectInitializer)
+{
+	bCreateNew = true;
+	bEditAfterNew = true;
+	SupportedClass = UIconsTracker::StaticClass();
+}
+UObject* UIconTrackerFactory::FactoryCreateNew(UClass* InClass, UObject* InParent, FName InName,
+	EObjectFlags Flags, UObject* Context, FFeedbackContext* Warn)
+{
+	return  NewObject<UIconsTracker>(InParent, UIconsTracker::StaticClass(), InName, Flags);
+}
+
+bool UIconTrackerFactory::CanCreateNew() const
+{
+	return true;
+}
diff --git a/Engine/Plugins/Experimental/UserToolBoxCore/Source/UserToolBoxCore/Private/UserToolBoxBaseBlueprint.cpp b/Engine/Plugins/Experimental/UserToolBoxCore/Source/UserToolBoxCore/Private/UserToolBoxBaseBlueprint.cpp
new file mode 100644
--- /dev/null
+++ b/Engine/Plugins/Experimental/UserToolBoxCore/Source/UserToolBoxCore/Private/UserToolBoxBaseBlueprint.cpp
@@ -0,0 +1,9 @@
+// Copyright Epic Games, Inc. All Rights Reserved.
+
+
+#include "UserToolBoxBaseBlueprint.h"
+
+inline void UUserToolBoxBaseBlueprint::Execute()
+{
+	Command();
+}
\ No newline at end of file
diff --git a/Engine/Plugins/Experimental/UserToolBoxCore/Source/UserToolBoxCore/Private/UserToolBoxCore.cpp b/Engine/Plugins/Experimental/UserToolBoxCore/Source/UserToolBoxCore/Private/UserToolBoxCore.cpp
new file mode 100644
--- /dev/null
+++ b/Engine/Plugins/Experimental/UserToolBoxCore/Source/UserToolBoxCore/Private/UserToolBoxCore.cpp
@@ -0,0 +1,28 @@
+// Copyright Epic Games, Inc. All Rights Reserved.
+
+#include "UserToolBoxCore.h"
+#include "UTBBaseCommand.h"
+#include "Customization/BaseCommandCustomization.h"
+#include "UObject/CoreRedirects.h"
+#include "PropertyEditorModule.h"
+#include "Modules/ModuleManager.h"
+
+#define LOCTEXT_NAMESPACE "FUserToolBoxCoreModule"
+
+void FUserToolBoxCoreModule::StartupModule()
+{
+
+	FPropertyEditorModule& PropertyEditorModule = FModuleManager::Get().LoadModuleChecked<FPropertyEditorModule>("PropertyEditor");
+	PropertyEditorModule.RegisterCustomClassLayout(UUTBBaseCommand::StaticClass()->GetFName(),FOnGetDetailCustomizationInstance::CreateStatic(&FBaseCommandCustomization::MakeInstance));
+	
+}
+
+void FUserToolBoxCoreModule::ShutdownModule()
+{
+	// This function may be called during shutdown to clean up your module.  For modules that support dynamic reloading,
+	// we call this function before unloading the module.
+}
+
+#undef LOCTEXT_NAMESPACE
+DEFINE_LOG_CATEGORY(LogUserToolBoxCore);
+IMPLEMENT_MODULE(FUserToolBoxCoreModule, UserToolBoxCore)
\ No newline at end of file
diff --git a/Engine/Plugins/Experimental/UserToolBoxCore/Source/UserToolBoxCore/Private/UserToolBoxFunctionLibrary.cpp b/Engine/Plugins/Experimental/UserToolBoxCore/Source/UserToolBoxCore/Private/UserToolBoxFunctionLibrary.cpp
new file mode 100644
--- /dev/null
+++ b/Engine/Plugins/Experimental/UserToolBoxCore/Source/UserToolBoxCore/Private/UserToolBoxFunctionLibrary.cpp
@@ -0,0 +1,48 @@
+// Copyright Epic Games, Inc. All Rights Reserved.
+
+
+#include "UserToolBoxFunctionLibrary.h"
+#include "UserToolBoxStyle.h"
+#include "Styling/SlateStyleRegistry.h"
+
+const FSlateBrush UUserToolBoxFunctionLibrary::DefaultBrush=FSlateBrush();
+const FSlateBrush& UUserToolBoxFunctionLibrary::GetBrushById(FString Name)
+{
+	const FSlateBrush* Ret=FUserToolBoxStyle::Get().GetBrush(FName(Name));
+	if (Ret==nullptr)
+	{
+		return DefaultBrush;
+	}
+	return *Ret;
+}
+
+TArray<FName> UUserToolBoxFunctionLibrary::GetAllSlateStyle()
+{
+	TArray<FName> Names;
+
+	FSlateStyleRegistry::IterateAllStyles(
+		[&](const ISlateStyle& Style)
+		{
+			Names.Add(Style.GetStyleSetName());
+			return true;
+		}
+	);
+	return Names;
+}
+
+const FSlateBrush& UUserToolBoxFunctionLibrary::GetBrushByStyleAndId(FName StyleName, FName Id)
+{
+	
+	
+	const ISlateStyle* Style=	FSlateStyleRegistry::FindSlateStyle(StyleName);
+	if (Style==nullptr)
+	{
+		return DefaultBrush; 
+	}
+	const FSlateBrush*	Brush= Style->GetBrush(Id);
+	if (Brush==nullptr)
+	{
+		return DefaultBrush;
+	}
+	return *Brush;
+}
diff --git a/Engine/Plugins/Experimental/UserToolBoxCore/Source/UserToolBoxCore/Private/UserToolBoxFunctionLibrary.h b/Engine/Plugins/Experimental/UserToolBoxCore/Source/UserToolBoxCore/Private/UserToolBoxFunctionLibrary.h
new file mode 100644
--- /dev/null
+++ b/Engine/Plugins/Experimental/UserToolBoxCore/Source/UserToolBoxCore/Private/UserToolBoxFunctionLibrary.h
@@ -0,0 +1,29 @@
+// Copyright Epic Games, Inc. All Rights Reserved.
+
+#pragma once
+
+#include "CoreMinimal.h"
+#include "Kismet/BlueprintFunctionLibrary.h"
+#include "UserToolBoxFunctionLibrary.generated.h"
+
+/**
+ * 
+ */
+UCLASS()
+class UUserToolBoxFunctionLibrary : public UBlueprintFunctionLibrary
+{
+	GENERATED_BODY()
+public:
+	static const FSlateBrush	DefaultBrush;
+	
+	UFUNCTION(BlueprintCallable, Category=" User Toolbox Library")
+	static const FSlateBrush& GetBrushById(FString Name);
+
+	UFUNCTION(BlueprintCallable, Category=" User Toolbox Library")
+	static TArray<FName> GetAllSlateStyle();
+
+	UFUNCTION(BlueprintCallable, Category=" User Toolbox Library")
+	static const FSlateBrush& GetBrushByStyleAndId(FName StyleName, FName Id);
+
+	
+};
diff --git a/Engine/Plugins/Experimental/UserToolBoxCore/Source/UserToolBoxCore/Private/UserToolBoxStyle.cpp b/Engine/Plugins/Experimental/UserToolBoxCore/Source/UserToolBoxCore/Private/UserToolBoxStyle.cpp
new file mode 100644
--- /dev/null
+++ b/Engine/Plugins/Experimental/UserToolBoxCore/Source/UserToolBoxCore/Private/UserToolBoxStyle.cpp
@@ -0,0 +1,103 @@
+// Copyright Epic Games, Inc. All Rights Reserved.
+
+#include "UserToolBoxStyle.h"
+#include "Styling/SlateStyle.h"
+#include "FileCache.h"
+#include "Styling/SlateStyleRegistry.h"
+#include "Framework/Application/SlateApplication.h"
+#include "Interfaces/IPluginManager.h"
+#include "Styling/StyleColors.h"
+
+TSharedPtr< FSlateStyleSet > FUserToolBoxStyle::StyleInstance = nullptr;
+TArray<FString> FUserToolBoxStyle::ExternalBrushIds;
+void FUserToolBoxStyle::Initialize()
+{
+	if (!StyleInstance.IsValid())
+	{
+		StyleInstance = Create();
+		FSlateStyleRegistry::RegisterSlateStyle(*StyleInstance);
+	}
+}
+
+void FUserToolBoxStyle::Shutdown()
+{
+	FSlateStyleRegistry::UnRegisterSlateStyle(*StyleInstance);
+	ensure(StyleInstance.IsUnique());
+	StyleInstance.Reset();
+}
+
+FName FUserToolBoxStyle::GetStyleSetName()
+{
+	static FName StyleSetName(TEXT("UserToolBoxStyle"));
+	return StyleSetName;
+}
+
+#define IMAGE_BRUSH( RelativePath, ... ) FSlateImageBrush( Style->RootToContentDir( RelativePath, TEXT(".png") ), __VA_ARGS__ )
+#define BOX_BRUSH( RelativePath, ... ) FSlateBoxBrush( Style->RootToContentDir( RelativePath, TEXT(".png") ), __VA_ARGS__ )
+#define BORDER_BRUSH( RelativePath, ... ) FSlateBorderBrush( Style->RootToContentDir( RelativePath, TEXT(".png") ), __VA_ARGS__ )
+#define TTF_FONT( RelativePath, ... ) FSlateFontInfo( Style->RootToContentDir( RelativePath, TEXT(".ttf") ), __VA_ARGS__ )
+#define OTF_FONT( RelativePath, ... ) FSlateFontInfo( Style->RootToContentDir( RelativePath, TEXT(".otf") ), __VA_ARGS__ )
+
+const FVector2D Icon16x16(16.0f, 16.0f);
+const FVector2D Icon15x15(16.0f, 16.0f);
+const FVector2D Icon30x30(16.0f, 16.0f);
+const FVector2D Icon60x60(16.0f, 16.0f);
+const FVector2D Icon20x20(20.0f, 20.0f);
+const FVector2D Icon40x40(40.0f, 40.0f);
+
+TSharedRef< FSlateStyleSet > FUserToolBoxStyle::Create()
+{
+	TSharedRef< FSlateStyleSet > Style = MakeShareable(new FSlateStyleSet("UserToolBoxStyle"));
+	Style->SetContentRoot(IPluginManager::Get().FindPlugin("UserToolBoxCore")->GetBaseDir() / TEXT("Resources"));
+	Style->Set("Palette.FirstHeader",new FSlateRoundedBoxBrush(FStyleColors::Dropdown, FVector4(4.0, 4.0, 0.0, 0.0)));
+	Style->Set("Palette.LastHeader",new FSlateRoundedBoxBrush(FStyleColors::Dropdown, FVector4(0.0, 0.0, 4.0, 4.0)));
+	Style->Set("Palette.Header",new FSlateRoundedBoxBrush(FStyleColors::Dropdown, FVector4(0.0, 0.0, 0.0, 0.0)));
+	Style->Set("Palette.UniqueHeader",new FSlateRoundedBoxBrush(FStyleColors::Dropdown, FVector4(4.0, 4.0, 4.0, 4.0)));
+	Style->Set("Palette.Body",new FSlateRoundedBoxBrush(FStyleColors::Recessed, FVector4(0.0, 0.0, 0.0, 0.0)));
+	
+	
+	return Style;
+}
+
+#undef IMAGE_BRUSH
+#undef BOX_BRUSH
+#undef BORDER_BRUSH
+#undef TTF_FONT
+#undef OTF_FONT
+
+void FUserToolBoxStyle::AddExternalImageBrushes(const TArray<FIconInfo>& IconInfos)
+{
+	if (StyleInstance==nullptr)
+	{
+		return;
+	}
+	for (const FIconInfo& Info:IconInfos)
+	{
+				StyleInstance->Set(FName(*Info.Id),new FSlateImageBrush(Info.Path,Info.IconSize));
+				ExternalBrushIds.Add(Info.Id);
+	}
+	ReloadTextures();
+}
+
+void FUserToolBoxStyle::ClearExternalImageBrushes()
+{
+	ExternalBrushIds.Empty();
+}
+
+TArray<FString> FUserToolBoxStyle::GetAvailableExternalImageBrushes()
+{
+	return ExternalBrushIds;
+}
+
+void FUserToolBoxStyle::ReloadTextures()
+{
+	if (FSlateApplication::IsInitialized())
+	{
+		FSlateApplication::Get().GetRenderer()->ReloadTextureResources();
+	}
+}
+
+const ISlateStyle& FUserToolBoxStyle::Get()
+{
+	return *StyleInstance;
+}
diff --git a/Engine/Plugins/Experimental/UserToolBoxCore/Source/UserToolBoxCore/Private/UserToolBoxSubsystem.cpp b/Engine/Plugins/Experimental/UserToolBoxCore/Source/UserToolBoxCore/Private/UserToolBoxSubsystem.cpp
new file mode 100644
--- /dev/null
+++ b/Engine/Plugins/Experimental/UserToolBoxCore/Source/UserToolBoxCore/Private/UserToolBoxSubsystem.cpp
@@ -0,0 +1,598 @@
+// Copyright Epic Games, Inc. All Rights Reserved.
+
+
+#include "UserToolBoxSubsystem.h"
+#include "AssetRegistry/AssetRegistryModule.h"
+#include "UTBBaseTab.h"
+#include "ToolMenus.h"
+#include "Widgets/Layout/SUniformGridPanel.h"
+#include "Widgets/Layout/SScrollBox.h"
+#include "Modules/ModuleManager.h"
+#include "AssetToolsModule.h"
+#include "WidgetBlueprintEditor.h"
+#include "AssetTypeActions_UTBEditorUtilityBlueprint.h"
+#include "EditorModeRegistry.h"
+#include "AssetTypeActions/AssetTypeActions_Blueprint.h"
+#include "WidgetBlueprint.h"
+#include "UserToolBoxStyle.h"
+#include "UTBBaseUITab.h"
+#include "IconsTracker.h"
+#include "LevelEditor.h"
+#include "NetworkMessage.h"
+#include "SLevelViewport.h"
+#include "StatusBarSubsystem.h"
+#include "UserToolBoxCore.h"
+#include "WidgetDrawerConfig.h"
+#include "Dialog/SCustomDialog.h"
+#include "Editor/UTBTabEditor.h"
+#include "Interfaces/IPluginManager.h"
+#include "Styling/StyleColors.h"
+#include "Widgets/Layout/SConstraintCanvas.h"
+#include "SVerticalTextBlock.h"
+#include "Editor/UTBEditorCommands.h"
+#include "Interfaces/IMainFrameModule.h"
+#include "Widgets/Docking/SDockTab.h"
+#include "Widgets/Input/SButton.h"
+#include "Widgets/Layout/SSpacer.h"
+#include "HAL/FileManager.h"
+class IMainFrameModule;
+EAssetTypeCategories::Type UUserToolboxSubsystem::AssetCategory=EAssetTypeCategories::Basic;
+
+class FAssetTypeActions_UTBBaseTab : public FAssetTypeActions_Base
+{
+public:
+	// IAssetTypeActions interface
+	virtual FText GetName() const override {return FText::FromString("User ToolBox Tab");};
+	virtual FColor GetTypeColor() const override{return FColor(169, 50, 0);};
+	virtual UClass* GetSupportedClass() const override {return UUserToolBoxBaseTab::StaticClass();};
+	virtual uint32 GetCategories() override{return UUserToolboxSubsystem::AssetCategory;};
+	virtual void OpenAssetEditor(const TArray<UObject*>& InObjects,
+		TSharedPtr<IToolkitHost> EditWithinLevelEditor) override
+	{
+		TArray<UUserToolBoxBaseTab*> Tabs;
+		for (UObject* Object:InObjects)
+		{
+			UUserToolBoxBaseTab* Tab=Cast<UUserToolBoxBaseTab>(Object);
+			if (IsValid(Tab))
+			{
+				Tabs.Add(Tab);
+			}
+		}
+		
+		for (UUserToolBoxBaseTab* Tab:Tabs)
+		{
+			TSharedRef< FUTBTabEditor > NewDataTableEditor( new FUTBTabEditor() );
+			NewDataTableEditor->InitUTBTabEditor( EToolkitMode::Standalone, EditWithinLevelEditor, Tab );
+		}
+	}
+	// End of IAssetTypeActions interface
+};
+
+
+class FAssetTypeActions_UIconTracker : public FAssetTypeActions_Base
+{
+public:
+	// IAssetTypeActions interface
+	virtual FText GetName() const override {return FText::FromString("Icon Tracker");};
+	virtual FColor GetTypeColor() const override{return FColor(100, 100, 0);};
+	virtual UClass* GetSupportedClass() const override {return UIconsTracker::StaticClass();};
+	virtual uint32 GetCategories() override{return UUserToolboxSubsystem::AssetCategory;};
+	// End of IAssetTypeActions interface
+};
+FAutoConsoleCommand RefreshIcons = FAutoConsoleCommand(
+	TEXT("UserToolbox.RefreshIcons"),
+	TEXT("RefreshIcons"),
+	FConsoleCommandDelegate::CreateLambda([]() {
+		GEditor->GetEditorSubsystem<UUserToolboxSubsystem>()->RefreshIcons();
+	}));
+FAutoConsoleCommand RefreshLevelViewportOverlay = FAutoConsoleCommand(
+	TEXT("UserToolbox.LevelViewportOverlay"),
+	TEXT("RefreshIcons"),
+	FConsoleCommandDelegate::CreateLambda([]() {
+		GEditor->GetEditorSubsystem<UUserToolboxSubsystem>()->UpdateLevelViewportWidget();
+	}));
+FAutoConsoleCommand RefreshDrawer = FAutoConsoleCommand(
+	TEXT("UserToolbox.RefreshDrawer"),
+	TEXT("Refresh the Drawer"),
+	FConsoleCommandDelegate::CreateLambda([]() {
+		GEditor->GetEditorSubsystem<UUserToolboxSubsystem>()->RegisterDrawer();
+	}));
+
+
+
+void UUserToolboxSubsystem::RegisterTabData()
+{
+	RefreshIcons();
+	UToolMenu* Menu = UToolMenus::Get()->ExtendMenu("LevelEditor.MainMenu.Window");
+	FToolMenuSection& Section = Menu->FindOrAddSection("UserToolBox");
+	Section.AddSubMenu("UserToolBox", FText::FromString("UserToolBox"), FText::FromString("Custom toolbox made by users"),
+		FNewMenuDelegate::CreateLambda([this](FMenuBuilder& MenuBuilder)
+		{
+			TArray<FAssetData> TabList=GetAvailableTabList();
+			for (const FAssetData& TabAssetData :TabList)
+			{
+				// need to handle collision
+				bool IsActive=TabAssetData.GetTagValueRef<bool>("bIsVisibleInWindowsMenu");
+				if (!IsActive)
+				{
+					continue;
+				}
+				FString TabName=TabAssetData.GetTagValueRef<FString>("Name");
+				MenuBuilder.AddMenuEntry(FText::FromString(TabName),FText::FromString(TabName),FSlateIcon(),
+				FUIAction(
+		FExecuteAction::CreateLambda([this,TabAssetData,TabName]()
+					{
+						FGlobalTabmanager::Get()->RegisterNomadTabSpawner(FName(TabName),FOnSpawnTab::CreateLambda([this,TabAssetData,TabName](const FSpawnTabArgs& arg)
+						{
+								
+							return SNew(SDockTab)
+							.TabRole(NomadTab)
+							[
+							GenerateTabUI(TabAssetData).ToSharedRef()
+							];
+							
+						}));
+						FGlobalTabmanager::Get()->TryInvokeTab(FName(TabName));
+						}),
+					FCanExecuteAction(),
+					FIsActionChecked()
+					)
+				)
+				;
+			}
+		}));
+	UpdateLevelViewportWidget();
+}
+TArray<FAssetData> UUserToolboxSubsystem::GetAvailableTabList()
+{
+	IAssetRegistry& AssetRegistry = FModuleManager::LoadModuleChecked<FAssetRegistryModule>(TEXT("AssetRegistry")).Get();
+	FARFilter Filter;
+	Filter.bIncludeOnlyOnDiskAssets = false;
+	Filter.ClassPaths.Add(UUserToolBoxBaseTab::StaticClass()->GetClassPathName());
+	Filter.bRecursivePaths = true;
+	Filter.bRecursiveClasses = true;
+	TArray<FAssetData> AssetDatas;
+	AssetRegistry.GetAssets(Filter,AssetDatas);
+	TArray<FAssetData> AvailableTabList;
+	for (FAssetData AssetData:AssetDatas)
+	{
+		UUserToolBoxBaseTab* Tab=Cast<UUserToolBoxBaseTab>(AssetData.GetAsset());
+		if (Tab==nullptr)
+		{
+			UE_LOG(LogTemp,Warning,TEXT("Unable to load %s"),*AssetData.PackageName.ToString())
+			continue;
+		}
+		AvailableTabList.Add(AssetData);
+	}
+	return AvailableTabList;
+}
+
+
+bool UUserToolboxSubsystem::PickAnIcon(FString& OutValue)
+{
+	IconOptions.Empty();
+	IconOptions.Add(MakeShareable(new FString("None")));
+	TArray<FString> BrushIds=FUserToolBoxStyle::GetAvailableExternalImageBrushes();
+
+	FString CurrentValue="";
+	for (const FString& BrushId:BrushIds)
+	{
+		IconOptions.Add(MakeShareable(new FString(BrushId)));
+	}
+	IconOptions.Sort([](const TSharedPtr<FString>& A, const TSharedPtr<FString>& B)
+	{
+		return *A<*B;
+	});
+
+	TSharedRef<SCustomDialog> OptionsWindow =
+	SNew(SCustomDialog)
+	.Title(FText::FromString("Icons"))
+	
+	.Buttons
+	(
+		{
+		SCustomDialog::FButton(FText::FromString("Ok")),
+		SCustomDialog::FButton(FText::FromString("Cancel")),
+		}
+	)
+	.Content()
+	[
+	SNew(SListView<TSharedPtr<FString>>)
+		.ListItemsSource(&IconOptions)
+		.OnSelectionChanged(SListView<TSharedPtr<FString>>::FOnSelectionChanged::CreateLambda([&CurrentValue](TSharedPtr<FString> Value, ESelectInfo::Type Type)
+		{
+			if (Value->Compare("None")==0)
+			{
+				CurrentValue="";
+			}
+			else
+			{
+				CurrentValue=*Value;
+			}
+		}))
+		.OnGenerateRow_Lambda([](TSharedPtr<FString> Node,const TSharedRef<STableViewBase>& OwnerTable)
+		{
+
+			TSharedPtr<SHorizontalBox> HBox=
+				SNew(SHorizontalBox)
+				+SHorizontalBox::Slot()
+				.AutoWidth()
+				.HAlign(HAlign_Left)
+				[
+					SNew(STextBlock).Text(FText::FromString(*Node))
+				];
+				
+			if ((*Node)!="None")
+			{
+				HBox->AddSlot()
+            				.HAlign(HAlign_Right)
+            				[
+            					SNew(SImage)
+            					.Image(FUserToolBoxStyle::Get().GetBrush(FName(*Node)))
+            				]	;
+			}
+			return SNew(STableRow<TSharedPtr<FString>>,OwnerTable)
+			[
+				HBox.ToSharedRef()
+			];
+		})
+	];
+	
+	const int32 PressedButtonIdx = OptionsWindow->ShowModal();
+	if (PressedButtonIdx != 0)
+	{
+		return false;
+	}
+	OutValue=CurrentValue;
+	return true;
+}
+
+void UUserToolboxSubsystem::RefreshIcons()
+{
+	FAssetRegistryModule& AssetRegistryModule = FModuleManager::LoadModuleChecked<FAssetRegistryModule>("AssetRegistry");
+	IAssetRegistry& AssetRegistry = AssetRegistryModule.Get();
+	FARFilter Filter;
+	Filter.bRecursiveClasses=true;
+	Filter.ClassPaths.Add(UIconsTracker::StaticClass()->GetClassPathName());
+	TArray<FAssetData> AssetDatas;
+	AssetRegistry.GetAssets(Filter,AssetDatas);
+	TArray<FIconInfo> IconDatas;
+	for (const FAssetData& AssetData:AssetDatas)
+	{
+		UIconsTracker* Asset=Cast<UIconsTracker>(AssetData.GetAsset());
+		if (IsValid(Asset))
+		{
+
+			FString Root;
+			if (AssetData.PackageName.ToString().StartsWith("/Game/"))
+			{// current project
+				Root=FPaths::ProjectDir();
+			}
+			else if ( AssetData.PackageName.ToString().StartsWith("/Engine/") )
+			{
+				Root=FPaths::EngineDir();
+			}
+			else
+			{
+				FString Tmp=AssetData.PackageName.ToString().RightChop(1);
+				int Index=0;
+				Tmp.FindChar('/',Index);
+				FString PluginName=Tmp.Left(Index);
+				TSharedPtr<IPlugin> Plugin=IPluginManager::Get().FindPlugin(PluginName);
+				if (Plugin!=nullptr)
+				{
+					Root=Plugin->GetBaseDir();
+				}
+					
+			}
+			if (!Root.IsEmpty())
+			{
+				for (const FIconFolderInfo& IconFolderInfo:Asset->IconFolderInfos)
+				{
+					
+					TArray<FString> IconPaths;
+					FString DirPath;
+					if (FPaths::IsRelative(IconFolderInfo.FolderPath.Path))
+					{
+						DirPath=Root / IconFolderInfo.FolderPath.Path ;
+					}
+					else
+					{
+						DirPath=IconFolderInfo.FolderPath.Path;
+					}
+					
+					if (!FPaths::DirectoryExists(DirPath))
+					{
+						UE_LOG(LogUserToolBoxCore,Error,TEXT(" directory doesn't exist: %s"),*DirPath);
+						continue;
+					}
+					
+					IFileManager::Get().FindFiles(IconPaths, *DirPath, TEXT(".png"));
+
+					for (const FString& IconPath:IconPaths)
+					{
+						if (FPaths::FileExists(DirPath / IconPath))
+						{
+							FIconInfo Info;
+							Info.Id=Asset->PrefixId+IconFolderInfo.PrefixId+FPaths::GetBaseFilename(IconPath);
+							Info.Path=DirPath / IconPath;
+							Info.IconSize=IconFolderInfo.IconSize;
+							IconDatas.Add(Info);
+					
+						}	
+					}
+				}	
+			}
+		}
+	}
+	FUserToolBoxStyle::ClearExternalImageBrushes();
+	FUserToolBoxStyle::AddExternalImageBrushes(IconDatas);
+}
+TSharedPtr<SWidget> UUserToolboxSubsystem::GenerateTabUI(const FAssetData Data, TSubclassOf<UUTBDefaultUITemplate> Ui,const FUITemplateParameters& Parameters)
+{
+	if (!Data.IsValid())
+	{
+		return SNullWidget::NullWidget;
+	}
+	UUserToolBoxBaseTab* Tab= Cast<UUserToolBoxBaseTab>(Data.GetAsset());
+	if (Tab==nullptr)
+	{
+		return SNullWidget::NullWidget;
+	}
+	RefreshIcons();
+	if (Data.IsValid())
+	{
+		return SNew(SUserToolBoxTabWidget)
+		.Tab(Data.GetAsset())
+		.UIOverride(Ui)
+		.UIParameters(Parameters);
+	}
+	return SNew(SDockTab);	
+	
+}
+void UUserToolboxSubsystem::Initialize(FSubsystemCollectionBase& Collection)
+{
+	Super::Initialize(Collection);
+	IAssetTools& AssetTools = FModuleManager::LoadModuleChecked<FAssetToolsModule>("AssetTools").Get();
+	AssetCategory = AssetTools.RegisterAdvancedAssetCategory(FName(TEXT("UserToolBox")), FText::FromString("User Tool Box"));
+	AssetTools.RegisterAssetTypeActions(MakeShared<FAssetTypeActions_UTBEditorBlueprint>());
+	AssetTools.RegisterAssetTypeActions(MakeShared<FAssetTypeActions_UTBBaseTab>());
+	AssetTools.RegisterAssetTypeActions(MakeShared<FAssetTypeActions_UIconTracker>());
+	
+	FAssetRegistryModule& AssetRegistryModule = FModuleManager::LoadModuleChecked<FAssetRegistryModule>("AssetRegistry");
+	IAssetRegistry& AssetRegistry = AssetRegistryModule.Get();
+	OnFileLoadHandle=AssetRegistry.OnFilesLoaded().AddLambda([this](){RegisterTabData();});
+	FUserToolBoxStyle::Initialize();
+	FUTBEditorCommands::Register();
+}
+void UUserToolboxSubsystem::Deinitialize()
+{
+	if (OnFileLoadHandle.IsValid())
+	{
+		FAssetRegistryModule& AssetRegistryModule = FModuleManager::LoadModuleChecked<FAssetRegistryModule>("AssetRegistry");
+		IAssetRegistry& AssetRegistry = AssetRegistryModule.Get();
+		AssetRegistry.OnFilesLoaded().Remove(OnFileLoadHandle);
+	}
+	FLevelEditorModule& LevelEditorModule = FModuleManager::GetModuleChecked<FLevelEditorModule>("LevelEditor");
+	TSharedPtr<SLevelViewport> ActiveLevelViewport = LevelEditorModule.GetFirstActiveLevelViewport();
+	if (ActiveLevelViewport.IsValid())
+	{
+		if (LevelViewportOverlayWdget.IsValid())
+		{
+			ActiveLevelViewport->RemoveOverlayWidget(LevelViewportOverlayWdget.ToSharedRef());
+			LevelViewportOverlayWdget.Reset();
+		}
+	}
+	if (LevelViewportOverlayWdget.IsValid())
+	{
+		LevelViewportOverlayWdget.Reset();
+	}
+	Super::Deinitialize();
+	if (FModuleManager::Get().IsModuleLoaded("MainFrame"))
+	{
+		const FUTBEditorCommands& Commands = FUTBEditorCommands::Get();
+
+		IMainFrameModule& MainFrame = FModuleManager::Get().LoadModuleChecked<IMainFrameModule>("MainFrame");
+
+		FUICommandList& ActionList = *MainFrame.GetMainFrameCommandBindings();
+
+		ActionList.UnmapAction(Commands.RenameSection);
+
+	}
+	
+	FUTBEditorCommands::Unregister();
+}
+
+void UUserToolboxSubsystem::RegisterDrawer()
+{
+	
+	UStatusBarSubsystem* StatusBarSubsystem=GEditor->GetEditorSubsystem<UStatusBarSubsystem>();
+	for (FName TabToUnregister:RegisteredDrawer)
+	{
+		//TODO find a way to clean the drawerbar
+	}
+	if (IsValid(StatusBarSubsystem))
+	{
+		TArray<FAssetData> TabList=GetAvailableTabList();
+		for (const FAssetData& TabAssetData :TabList)
+		{
+			FString TabName=TabAssetData.GetTagValueRef<FString>("Name");
+			if (RegisteredDrawer.Contains(FName(TabName)))
+			{
+				continue;	
+			}
+			bool IsVisible=TabAssetData.GetTagValueRef<bool>("bIsVisibleInDrawer");
+			if (IsVisible)
+			{
+				FName TabNameAsFName(TabName);
+				FWidgetDrawerConfig WidgetDrawerConfig(TabNameAsFName);
+				WidgetDrawerConfig.ButtonText=FText::FromString(TabName);
+				WidgetDrawerConfig.ToolTipText=FText::FromString(TabName);
+				WidgetDrawerConfig.GetDrawerContentDelegate=FOnGetContent::CreateLambda([this,TabAssetData]()
+				{
+					return GenerateTabUI(TabAssetData, nullptr).ToSharedRef();
+					
+				});
+				WidgetDrawerConfig.OnDrawerDismissedDelegate=FOnStatusBarDrawerDismissed::CreateLambda([]( const TSharedPtr<SWidget>& widget)
+				{
+					UE_LOG(LogTemp, Warning, TEXT("Dismissed"))
+				});
+				WidgetDrawerConfig.OnDrawerOpenedDelegate=FOnStatusBarDrawerOpened::CreateLambda([](FName StatusBarName)
+				{
+					UE_LOG(LogTemp, Warning, TEXT("openned %s"),*StatusBarName.ToString());
+				});
+				RegisteredDrawer.Add(FName(TabName));
+				StatusBarSubsystem->RegisterDrawer("LevelEditor.StatusBar",MoveTemp(WidgetDrawerConfig));				
+			}
+				
+		}
+		
+		
+	}
+}
+
+void UUserToolboxSubsystem::UpdateLevelViewportWidget()
+{
+	FLevelEditorModule& LevelEditorModule = FModuleManager::GetModuleChecked<FLevelEditorModule>("LevelEditor");
+	TSharedPtr<SLevelViewport> ActiveLevelViewport = LevelEditorModule.GetFirstActiveLevelViewport();
+	RegisterDrawer();
+	TArray<FAssetData> TabToAdd;
+	if (ActiveLevelViewport.IsValid())
+	{
+		if (LevelViewportOverlayWdget.IsValid())
+		{
+			ActiveLevelViewport->RemoveOverlayWidget(LevelViewportOverlayWdget.ToSharedRef());
+			LevelViewportOverlayWdget.Reset();
+		}
+		TArray<FAssetData> TabList=GetAvailableTabList();
+		for (const FAssetData& TabAssetData :TabList)
+		{
+			// need to handle collision
+			bool IsVisible=TabAssetData.GetTagValueRef<bool>("bIsVisibleInViewportOverlay");
+			if (IsVisible)
+			{
+				TabToAdd.Add(TabAssetData);
+			}
+		}
+		if (!TabToAdd.IsEmpty())
+		{
+			TSharedPtr<SVerticalBox> VerticalBox;
+			LevelViewportOverlayWdget=
+				SNew(SConstraintCanvas)
+				 +SConstraintCanvas::Slot()
+				.Anchors(FAnchors(0.0f,.0f))
+				.Alignment(FVector2d(0.0,0.0f))
+				.Offset(FMargin(0.0f,138.0f,0.0,0.0f))
+				.AutoSize(true)
+				 [
+					 SNew(SBox)
+					 .MaxDesiredWidth(250)
+					 [
+						 SAssignNew(VerticalBox,SVerticalBox)
+				 	]
+				 ];
+			int Index=0;
+			for (FAssetData& TabAssetData:TabToAdd)
+			{
+				FName HeaderStyle="Palette.Header";
+				if (TabToAdd.Num()==1)
+				{
+					HeaderStyle="Palette.UniqueHeader";	
+				}
+				else
+				{
+					if (Index==0)
+					{
+						HeaderStyle="Palette.FirstHeader";
+					}
+					if (Index==TabToAdd.Num()-1)
+					{
+						HeaderStyle="Palette.LastHeader";
+					}
+				}
+				FString TabName=TabAssetData.GetTagValueRef<FString>("Name");
+				FUITemplateParameters Params;
+				Params.bForceHideSettingsButton=true;
+				Params.bPrefixSectionWithTabName=false;
+				TSharedPtr<FCurveSequence> CurveSequence= MakeShared<FCurveSequence>(0,.3,ECurveEaseFunction::CubicIn);
+				const FSlateRenderTransform Rotate90(FQuat2D(FMath::DegreesToRadians(90.f)));
+				VerticalBox->AddSlot()
+				 .AutoHeight()
+				 .HAlign(HAlign_Fill)
+				 .VAlign(VAlign_Top)
+				[
+					SNew(SBox)
+					.MinDesiredHeight(100)
+					[
+						
+							SNew(SHorizontalBox)
+							+SHorizontalBox::Slot()
+							.HAlign(HAlign_Left)
+							.AutoWidth()
+							[
+								SNew(SBox)
+								.WidthOverride(20)
+								.Clipping(EWidgetClipping::Inherit)
+								[
+									SNew(SButton)
+									.ContentPadding(FMargin(5,10,0,10))
+									.Clipping(EWidgetClipping::Inherit)
+									.OnClicked_Lambda([CurveSequence,this]()
+									{
+										if (CurveSequence->GetSequenceTime()<0.1)
+										{
+											CurveSequence->Play(LevelViewportOverlayWdget->AsShared());
+										}
+										if (CurveSequence->GetSequenceTime()>0.9)
+										{
+											CurveSequence->PlayReverse(LevelViewportOverlayWdget->AsShared());
+										}
+										return FReply::Handled();
+									})
+									[
+										SNew(SVerticalTextBlock)
+										.Clipping(EWidgetClipping::Inherit)
+										.Text(FText::FromString(TabName))
+										.Font(FAppStyle::Get().GetFontStyle("BoldFont"))
+										.ColorAndOpacity(FAppStyle::Get().GetSlateColor("Colors.Foreground"))
+
+									]
+								]
+							]
+							+SHorizontalBox::Slot()
+							.HAlign(HAlign_Fill)
+							[
+								SNew(SBorder)
+								.Padding(0)
+								.DesiredSizeScale_Lambda([CurveSequence]()
+								{
+									return FVector2d(CurveSequence->GetLerp(),1);
+								})
+								.Visibility_Lambda([CurveSequence]()
+								{
+									return CurveSequence->GetLerp()<0.1?EVisibility::Collapsed:EVisibility::Visible;
+								})
+								.BorderImage( FAppStyle::GetBrush("ToolPanel.GroupBorder") )
+								
+								.BorderBackgroundColor(FAppStyle::Get().GetSlateColor("Colors.Recessed"))
+								[
+									GenerateTabUI(TabAssetData, UUTBVerticalToolBarTabUI::StaticClass(),Params).ToSharedRef()
+								]
+							]
+					]
+				];
+				VerticalBox->AddSlot()
+				.AutoHeight()
+				[
+					SNew(SSpacer).Size(FVector(20.0,20.0,20.0))
+				];
+				Index++;
+			}
+			if (LevelViewportOverlayWdget.IsValid())
+			{
+				ActiveLevelViewport->AddOverlayWidget(LevelViewportOverlayWdget.ToSharedRef());	
+			}
+		}
+	}
+}
+
diff --git a/Engine/Plugins/Experimental/UserToolBoxCore/Source/UserToolBoxCore/Public/Editor/IUTBTabEditor.h b/Engine/Plugins/Experimental/UserToolBoxCore/Source/UserToolBoxCore/Public/Editor/IUTBTabEditor.h
new file mode 100644
--- /dev/null
+++ b/Engine/Plugins/Experimental/UserToolBoxCore/Source/UserToolBoxCore/Public/Editor/IUTBTabEditor.h
@@ -0,0 +1,14 @@
+// Copyright Epic Games, Inc. All Rights Reserved.
+
+#pragma once
+
+#include "CoreMinimal.h"
+#include "Toolkits/AssetEditorToolkit.h"
+/**
+ * 
+ */
+class USERTOOLBOXCORE_API IUTBTabEditor : public FAssetEditorToolkit
+{
+
+	
+};
diff --git a/Engine/Plugins/Experimental/UserToolBoxCore/Source/UserToolBoxCore/Public/Editor/SCommandPickMenu.h b/Engine/Plugins/Experimental/UserToolBoxCore/Source/UserToolBoxCore/Public/Editor/SCommandPickMenu.h
new file mode 100644
--- /dev/null
+++ b/Engine/Plugins/Experimental/UserToolBoxCore/Source/UserToolBoxCore/Public/Editor/SCommandPickMenu.h
@@ -0,0 +1,61 @@
+﻿// Copyright Epic Games, Inc. All Rights Reserved.
+
+#pragma once
+
+#include "CoreMinimal.h"
+#include "UTBBaseTab.h"
+#include "Widgets/Views/STreeView.h"
+/**
+ * 
+ */
+DECLARE_DELEGATE_OneParam(FOnCommandSelectedChanged, UClass*)
+class USERTOOLBOXCORE_API SCommandPickMenuWidget : public SCompoundWidget
+{
+	struct FCommandTreeNode
+	{
+		bool bIsCategory= false;
+		UClass*									Class;
+		FString									Name;
+		TArray<TSharedPtr<FCommandTreeNode>>	Children;
+		TArray<TSharedPtr<FCommandTreeNode>>	FilteredChildren;
+	};
+
+
+public:
+	SLATE_BEGIN_ARGS( SCommandPickMenuWidget )
+		 
+			:_OnCommandSelectionChanged()
+		{
+		
+		}
+	SLATE_EVENT(FOnCommandSelectedChanged, OnCommandSelectionChanged)
+		/** Sets the text content for this editable text widget */
+		
+		
+	SLATE_END_ARGS()
+
+
+	SCommandPickMenuWidget();
+	~SCommandPickMenuWidget();
+
+	void Construct( const FArguments& InArgs );
+
+	
+private:
+	void ReconstructCommandTree();
+
+	void UpdateFilteredTreeNodes(FString Text);
+	void OnSetExpansionRecursive(TSharedPtr<FCommandTreeNode> InTreeNode, bool bInIsItemExpanded);
+
+
+	FOnCommandSelectedChanged OnCommandSelectionChanged;
+	TArray<TSharedPtr<FCommandTreeNode>> CommandTree;
+	TArray<TSharedPtr<FCommandTreeNode>> CommandTreeAllNodes;
+	
+	TMap<TSharedPtr<FCommandTreeNode>,TArray<TSharedPtr<FCommandTreeNode>>> LeafsWithParentHierarchy;// reverse tree for seaching
+	
+	TSharedPtr<STreeView<TSharedPtr<FCommandTreeNode>>> TreeViewPtr;
+
+	UClass* Selection;
+};
+
diff --git a/Engine/Plugins/Experimental/UserToolBoxCore/Source/UserToolBoxCore/Public/Editor/SCommandWithOptionsDetailsView.h b/Engine/Plugins/Experimental/UserToolBoxCore/Source/UserToolBoxCore/Public/Editor/SCommandWithOptionsDetailsView.h
new file mode 100644
--- /dev/null
+++ b/Engine/Plugins/Experimental/UserToolBoxCore/Source/UserToolBoxCore/Public/Editor/SCommandWithOptionsDetailsView.h
@@ -0,0 +1,38 @@
+﻿// Copyright Epic Games, Inc. All Rights Reserved.
+
+#pragma once
+
+#include "CoreMinimal.h"
+#include "UTBBaseTab.h"
+#include "IDetailsView.h"
+/**
+ * 
+ */
+DECLARE_DELEGATE_OneParam(FOnObjectPropertyModified,UUTBBaseCommand*);
+class USERTOOLBOXCORE_API SCommandDetailsView : public SCompoundWidget
+{
+public:
+	SLATE_BEGIN_ARGS( SCommandDetailsView )
+		: _Command() 
+		{
+		
+		}
+
+		/** Sets the text content for this editable text widget */
+		SLATE_ATTRIBUTE( UUTBBaseCommand*, Command )
+		SLATE_EVENT( FOnObjectPropertyModified, OnObjectPropertyModified )
+	
+	SLATE_END_ARGS()
+
+
+	SCommandDetailsView();
+	~SCommandDetailsView();
+	
+	void Construct( const FArguments& InArgs );
+	void SetObject(UUTBBaseCommand* Command);
+	
+private:
+	TWeakObjectPtr<UUTBBaseCommand> Command;
+	TSharedPtr<IDetailsView>		CommandDetailsView;
+	FOnObjectPropertyModified		OnObjectPropertyModified;
+};
diff --git a/Engine/Plugins/Experimental/UserToolBoxCore/Source/UserToolBoxCore/Public/Editor/SCommandWithOptionsWidget.h b/Engine/Plugins/Experimental/UserToolBoxCore/Source/UserToolBoxCore/Public/Editor/SCommandWithOptionsWidget.h
new file mode 100644
--- /dev/null
+++ b/Engine/Plugins/Experimental/UserToolBoxCore/Source/UserToolBoxCore/Public/Editor/SCommandWithOptionsWidget.h
@@ -0,0 +1,97 @@
+﻿// Copyright Epic Games, Inc. All Rights Reserved.
+
+#pragma once
+
+#include "CoreMinimal.h"
+#include "UTBBaseCommand.h"
+#include "UTBBaseTab.h"
+#include "UTBTabEditor.h"
+#include "Widgets/Input/SButton.h"
+
+/**
+ * 
+ */
+class USERTOOLBOXCORE_API SButtonCommandWidget : public SButton
+{
+
+public:
+	SLATE_BEGIN_ARGS( SButtonCommandWidget )
+	:_Content()
+	{
+		
+		
+	}
+	/** Slot for this button's content (optional) */
+	SLATE_DEFAULT_SLOT( FArguments, Content )
+	SLATE_EVENT( FOnClicked, OnClicked )
+	
+	SLATE_ATTRIBUTE( UUTBBaseCommand*, Command )
+	
+SLATE_END_ARGS()
+	SButtonCommandWidget();
+	~SButtonCommandWidget();
+
+	void Construct(const FArguments& InArgs );
+	virtual FReply OnDragDetected(const FGeometry& MyGeometry, const FPointerEvent& MouseEvent) override;
+	virtual FReply OnMouseButtonDown(const FGeometry& MyGeometry, const FPointerEvent& MouseEvent) override;
+	
+
+	
+private:
+	TWeakObjectPtr<UUTBBaseCommand> Command;
+	FOnContextMenuOpening OnContextMenuOpening;
+	
+};
+class USERTOOLBOXCORE_API SCommandWidget : public SCompoundWidget
+{
+public:
+	DECLARE_DELEGATE_OneParam(FOnTabChangedNotify, UUserToolBoxBaseTab*)
+	DECLARE_DELEGATE_ThreeParams(FOnCommandChanged, UUTBBaseCommand*, FString , int )
+	DECLARE_DELEGATE_OneParam(FOnCommandSelected, UUTBBaseCommand*)
+	DECLARE_DELEGATE_OneParam(FOnCommandDeleted,  UUTBBaseCommand* )
+	DECLARE_DELEGATE_ThreeParams(FOnDropCommand, UUTBBaseCommand*, FString , int )
+
+	SLATE_BEGIN_ARGS( SCommandWidget )
+		: _Command() 
+		, _IndexInSection(-1)
+		, _SectionName("")
+		{
+		
+		}
+
+		/** Sets the text content for this editable text widget */
+		SLATE_ATTRIBUTE( UUTBBaseCommand*, Command )
+		SLATE_ATTRIBUTE( int, IndexInSection )
+		SLATE_ATTRIBUTE( FString, SectionName )
+		SLATE_ATTRIBUTE( TSharedPtr<FUTBTabEditor>, Editor )
+		SLATE_EVENT(FOnTabChangedNotify, OnTabChanged)
+		SLATE_EVENT(FOnCommandChanged, OnCommandChanged)
+		SLATE_EVENT(FOnCommandSelected, OnCommandSelected)
+		SLATE_EVENT(FOnCommandDeleted, OnCommandDeleted)
+		SLATE_EVENT(FOnDropCommand, OnDropCommand)
+		SLATE_EVENT(FOnContextMenuOpening, OnContextMenuOpening)
+	SLATE_END_ARGS()
+
+	SCommandWidget();
+	~SCommandWidget();
+	virtual FReply OnDragDetected(const FGeometry& MyGeometry, const FPointerEvent& MouseEvent) override;
+	virtual FReply OnMouseButtonDown(const FGeometry& MyGeometry, const FPointerEvent& MouseEvent) override;
+	void Construct( const FArguments& InArgs );
+
+	void OnObjectsReplacedCallback(const TMap<UObject*, UObject*>& ReplacementMap);
+	virtual FReply OnMouseButtonUp(const FGeometry& MyGeometry, const FPointerEvent& MouseEvent) override;
+	TSharedPtr<SWidget> CreateContextMenu();
+private:
+	void RegisterContextMenu();
+	void RebuildWidget();
+	TWeakPtr<FUTBTabEditor> Editor;
+	FOnTabChangedNotify OnTabChanged;
+	FOnCommandChanged OnCommandChanged;
+	FOnCommandSelected OnCommandSelected;
+	FOnCommandDeleted OnCommandDeleted;
+	FOnDropCommand OnDropCommand;
+	FOnContextMenuOpening OnContextMenuOpening;
+	TWeakObjectPtr<UUTBBaseCommand> Command;
+	int IndexInSection;
+	FString SectionName;
+};
diff --git a/Engine/Plugins/Experimental/UserToolBoxCore/Source/UserToolBoxCore/Public/Editor/UTBTabEditor.h b/Engine/Plugins/Experimental/UserToolBoxCore/Source/UserToolBoxCore/Public/Editor/UTBTabEditor.h
new file mode 100644
--- /dev/null
+++ b/Engine/Plugins/Experimental/UserToolBoxCore/Source/UserToolBoxCore/Public/Editor/UTBTabEditor.h
@@ -0,0 +1,107 @@
+// Copyright Epic Games, Inc. All Rights Reserved.
+
+#pragma once
+
+#include "CoreMinimal.h"
+#include "IUTBTabEditor.h"
+#include "EditorUndoClient.h"
+#include "SCommandWithOptionsDetailsView.h"
+#include "UTBBaseTab.h"
+#include "Widgets/Views/SListView.h"
+/**
+ * 
+ */
+
+struct USERTOOLBOXCORE_API FCommandDragDropOperation : public FDragDropOperation
+{
+public:
+	DRAG_DROP_OPERATOR_TYPE(FCommandDragDropOperation, FDragDropOperation);
+	using FDragDropOperation::Construct;
+	FCommandDragDropOperation();
+	virtual ~FCommandDragDropOperation();
+
+	static TSharedRef<FCommandDragDropOperation> New(UUTBBaseCommand* Command);
+
+	
+	TWeakObjectPtr<UUTBBaseCommand> Command;
+	
+};
+
+
+class USERTOOLBOXCORE_API FUTBTabEditor : public IUTBTabEditor
+		, public FEditorUndoClient
+{
+public:
+	FUTBTabEditor();
+	~FUTBTabEditor();
+	
+	DECLARE_MULTICAST_DELEGATE_OneParam(FOnCurrentSectionChange, FString);
+	DECLARE_MULTICAST_DELEGATE(FOnSectionListChange);
+	
+	virtual FLinearColor GetWorldCentricTabColorScale() const override;
+	virtual void RegisterTabSpawners(const TSharedRef<FTabManager>& TabManager) override;
+	virtual void UnregisterTabSpawners(const TSharedRef<FTabManager>& TabManager) override;
+	virtual FName GetToolkitFName() const override;
+	virtual FText GetBaseToolkitName() const override;
+	virtual FString GetWorldCentricTabPrefix() const override;
+	virtual void PostUndo(bool bSuccess) override;
+	virtual void PostRedo(bool bSuccess) override;
+
+	virtual void InitUTBTabEditor(const EToolkitMode::Type Mode,const TSharedPtr<IToolkitHost>& InitToolkitHost, UUserToolBoxBaseTab* Tab);
+	FString GetCurrentSection();
+
+	FOnCurrentSectionChange	OnCurrentSectionChange;
+	FOnSectionListChange OnSectionListChange;
+
+	void AddCommandToCurrentSection(UUTBBaseCommand* InCommand);
+	void RemoveCommandFromCurrentSection(UUTBBaseCommand* InCommand, bool ShouldRebuild=false);
+	
+	void AddCommandAt(UUTBBaseCommand* InCommand, FString Section,int Index);
+	void SetCurrentCommandSelection(UUTBBaseCommand* Command);
+	void NotifyTabHasChanged(UUserToolBoxBaseTab* Tab);
+	void UpdateCommand(UUTBBaseCommand* Command);
+	void UpdateCurrentSection();
+	void CreateNewSection();
+	void MoveSectionAfterExistingSection(FString SectionToMove,FString SectionToTarget);
+	void RefreshSectionNames();
+	void RemoveCurrentSection();
+	void EditActiveSectionName();
+	void RenameActiveSection(FString NewName);
+	void BindCommands();
+protected:
+
+	void CreateOrRebuildUTBTabTabWidget();
+	void CreateOrRebuildUTBTabSectionTabWidget();
+	TSharedPtr<SWidget> CreateUTBTabDetailsTabWidget();
+	TSharedPtr<SWidget> CreateUTBCommandTabWidget();
+	TSharedPtr<SWidget> CreateUTBCommandDetailsTabWidget();
+
+	void SetCurrentSection(FString SectionName);
+	TArray<TSharedPtr<FString>> SectionNames;
+	
+
+
+
+
+	
+
+
+	
+	static const FName UTBTabTabId;
+	static const FName UTBTabSectionTabId;
+	static const FName UTBTabDetailsTabId;
+	static const FName UTBCommandTabId;
+	static const FName UTBCommandDetailsTabId;
+	FString CurrentSection;
+	FAssetData CurrentSelectedCommand;
+	TSharedPtr<SVerticalBox>	UUTBTabTabWidget;
+	TSharedPtr<SVerticalBox>    UUTBTabSectionTabWidget;
+	TSharedPtr<SWidget>		    UUTBTabDetailsTabWidget;
+	TSharedPtr<SWidget>		    UUTBCommandTabWidget;
+	TSharedPtr<SWidget>		    UUTBCommandDetailsTabWidget;
+	
+	TSharedPtr<SCommandDetailsView>	CommandDetailsView;
+	TSharedPtr<SVerticalBox> SectionVerticalBox;
+	TSharedPtr<SListView<TSharedPtr<FString>>> SectionListView;
+};
+
diff --git a/Engine/Plugins/Experimental/UserToolBoxCore/Source/UserToolBoxCore/Public/IconsTracker.h b/Engine/Plugins/Experimental/UserToolBoxCore/Source/UserToolBoxCore/Public/IconsTracker.h
new file mode 100644
--- /dev/null
+++ b/Engine/Plugins/Experimental/UserToolBoxCore/Source/UserToolBoxCore/Public/IconsTracker.h
@@ -0,0 +1,46 @@
+// Copyright Epic Games, Inc. All Rights Reserved.
+
+#pragma once
+
+#include "CoreMinimal.h"
+#include "Engine/DataAsset.h"
+#include "Engine/DataTable.h"
+#include "IconsTracker.generated.h"
+
+USTRUCT()
+struct FIconInfo
+{
+	GENERATED_BODY()
+	UPROPERTY(EditAnywhere , Category="Icon Tracker")
+	FString Id;
+	UPROPERTY(EditAnywhere , Category="Icon Tracker")
+	FString Path;
+	UPROPERTY(EditAnywhere , Category="Icon Tracker")
+	FVector2D IconSize=FVector2D(30.0f);
+	
+	
+};
+USTRUCT()
+struct FIconFolderInfo : public FTableRowBase
+{
+	GENERATED_BODY()
+	UPROPERTY(EditAnywhere , Category="Icon Tracker")
+	FDirectoryPath FolderPath;
+	UPROPERTY(EditAnywhere , Category="Icon Tracker")
+	FString PrefixId;
+	UPROPERTY(EditAnywhere , Category="Icon Tracker")
+	FVector2D IconSize=FVector2D(30.0f);
+};
+
+UCLASS()
+class USERTOOLBOXCORE_API UIconsTracker : public UDataAsset
+{
+	GENERATED_BODY()
+
+	public:
+	UPROPERTY(EditAnywhere , Category="Icon Tracker")
+	TArray<FIconFolderInfo>		IconFolderInfos;
+	UPROPERTY(EditAnywhere , Category="Icon Tracker")
+	FString					PrefixId;
+	
+};
diff --git a/Engine/Plugins/Experimental/UserToolBoxCore/Source/UserToolBoxCore/Public/SVerticalTextBlock.h b/Engine/Plugins/Experimental/UserToolBoxCore/Source/UserToolBoxCore/Public/SVerticalTextBlock.h
new file mode 100644
--- /dev/null
+++ b/Engine/Plugins/Experimental/UserToolBoxCore/Source/UserToolBoxCore/Public/SVerticalTextBlock.h
@@ -0,0 +1,133 @@
+// Copyright Epic Games, Inc. All Rights Reserved.
+
+#pragma once
+
+#include "CoreMinimal.h"
+#include "Widgets/Text/STextBlock.h"
+class SVerticalTextBlock : public STextBlock
+{
+public:
+
+	SLATE_BEGIN_ARGS( SVerticalTextBlock )
+		: _Text()
+		, _TextStyle( &FCoreStyle::Get().GetWidgetStyle<FTextBlockStyle>( "NormalText" ) )
+		, _Font()
+		, _StrikeBrush()
+		, _ColorAndOpacity()
+		, _ShadowOffset()
+		, _ShadowColorAndOpacity()
+		, _HighlightColor()
+		, _HighlightShape()
+		, _HighlightText()
+		, _WrapTextAt(0.0f)
+		, _AutoWrapText(false)
+		, _WrappingPolicy(ETextWrappingPolicy::DefaultWrapping)
+		, _TransformPolicy()
+		, _Margin()
+		, _LineHeightPercentage(1.0f)
+		, _Justification(ETextJustify::Left)
+		, _MinDesiredWidth(0.0f)
+		, _TextShapingMethod()
+		, _TextFlowDirection()
+		, _LineBreakPolicy()
+		, _OverflowPolicy()
+		, _SimpleTextMode(false)
+	{
+		_Clipping = EWidgetClipping::OnDemand;
+	}
+
+	/** The text displayed in this text block */
+	SLATE_ATTRIBUTE( FText, Text )
+
+	/** Pointer to a style of the text block, which dictates the font, color, and shadow options. */
+	SLATE_STYLE_ARGUMENT( FTextBlockStyle, TextStyle )
+
+	/** Sets the font used to draw the text */
+	SLATE_ATTRIBUTE( FSlateFontInfo, Font )
+
+	/** Sets the brush used to strike through the text */
+	SLATE_ATTRIBUTE( const FSlateBrush*, StrikeBrush )
+
+	/** Text color and opacity */
+	SLATE_ATTRIBUTE( FSlateColor, ColorAndOpacity )
+
+	/** Drop shadow offset in pixels */
+	SLATE_ATTRIBUTE( FVector2D, ShadowOffset )
+
+	/** Shadow color and opacity */
+	SLATE_ATTRIBUTE( FLinearColor, ShadowColorAndOpacity )
+
+	/** The color used to highlight the specified text */
+	SLATE_ATTRIBUTE( FLinearColor, HighlightColor )
+
+	/** The brush used to highlight the specified text */
+	SLATE_ATTRIBUTE( const FSlateBrush*, HighlightShape )
+
+	/** Highlight this text in the text block */
+	SLATE_ATTRIBUTE( FText, HighlightText )
+
+	/** Whether text wraps onto a new line when it's length exceeds this width; if this value is zero or negative, no wrapping occurs. */
+	SLATE_ATTRIBUTE( float, WrapTextAt )
+
+	/** Whether to wrap text automatically based on the widget's computed horizontal space.  IMPORTANT: Using automatic wrapping can result
+		in visual artifacts, as the the wrapped size will computed be at least one frame late!  Consider using WrapTextAt instead.  The initial 
+		desired size will not be clamped.  This works best in cases where the text block's size is not affecting other widget's layout. */
+	SLATE_ATTRIBUTE( bool, AutoWrapText )
+
+	/** The wrapping policy to use */
+	SLATE_ATTRIBUTE( ETextWrappingPolicy, WrappingPolicy )
+
+	/** The transform policy to use */
+	SLATE_ATTRIBUTE( ETextTransformPolicy, TransformPolicy )
+
+	/** The amount of blank space left around the edges of text area. */
+	SLATE_ATTRIBUTE( FMargin, Margin )
+
+	/** The amount to scale each lines height by. */
+	SLATE_ATTRIBUTE( float, LineHeightPercentage )
+
+	/** How the text should be aligned with the margin. */
+	SLATE_ATTRIBUTE( ETextJustify::Type, Justification )
+
+	/** Minimum width that a text block should be */
+	SLATE_ATTRIBUTE( float, MinDesiredWidth )
+
+	/** Which text shaping method should we use? (unset to use the default returned by GetDefaultTextShapingMethod) */
+	SLATE_ARGUMENT( TOptional<ETextShapingMethod>, TextShapingMethod )
+		
+	/** Which text flow direction should we use? (unset to use the default returned by GetDefaultTextFlowDirection) */
+	SLATE_ARGUMENT( TOptional<ETextFlowDirection>, TextFlowDirection )
+
+	/** The iterator to use to detect appropriate soft-wrapping points for lines (or null to use the default) */
+	SLATE_ARGUMENT( TSharedPtr<IBreakIterator>, LineBreakPolicy )
+
+	/** Determines what happens to text that is clipped and doesn't fit within the clip rect for this widget */
+	SLATE_ARGUMENT(TOptional<ETextOverflowPolicy>, OverflowPolicy)
+
+	/**
+	 * If this is enabled, text shaping, wrapping, justification are disabled in favor of much faster text layout and measurement.
+	 * This feature is suitable for numbers and text that changes often and impact performance.
+	 * Enabling this setting may cause certain languages (such as Right to left languages) to not display properly.
+	 */
+	SLATE_ARGUMENT( bool, SimpleTextMode )
+
+	/** Called when this text is double clicked */
+	SLATE_EVENT(FPointerEventHandler, OnDoubleClicked)
+
+SLATE_END_ARGS()
+
+/** Constructor */
+SVerticalTextBlock();
+
+	/** Destructor */
+	~SVerticalTextBlock();
+
+	/**
+	 * Construct this widget
+	 *
+	 * @param	InArgs	The declaration data for this widget
+	 */
+	void Construct( const FArguments& InArgs );
+
+	virtual FVector2D ComputeDesiredSize(const float) const override;
+};
\ No newline at end of file
diff --git a/Engine/Plugins/Experimental/UserToolBoxCore/Source/UserToolBoxCore/Public/UTBBaseCommand.h b/Engine/Plugins/Experimental/UserToolBoxCore/Source/UserToolBoxCore/Public/UTBBaseCommand.h
new file mode 100644
--- /dev/null
+++ b/Engine/Plugins/Experimental/UserToolBoxCore/Source/UserToolBoxCore/Public/UTBBaseCommand.h
@@ -0,0 +1,53 @@
+// Copyright Epic Games, Inc. All Rights Reserved.
+
+#pragma once
+
+#include "CoreMinimal.h"
+#include "Components/Widget.h"
+#include "Framework/Commands/InputChord.h"
+#include "UTBBaseCommand.generated.h"
+
+
+/**
+ * UUTBBaseCommand is the base class for any command of the usertoolbox framework.
+ * Inherit directly from this class if you want to create a new command implemented in c++
+ */
+UCLASS(Blueprintable, Abstract,EditInlineNew)
+class USERTOOLBOXCORE_API UUTBBaseCommand : public UObject
+{
+	GENERATED_BODY()
+	public:
+	/** The name of the command */
+	UPROPERTY(EditAnywhere, BlueprintReadWrite, Category="Command Description", AssetRegistrySearchable)
+	FString Name="Default name";
+	/** The icon path for the command */
+	UPROPERTY(EditAnywhere, BlueprintReadWrite, Category="Command Description")
+	FString IconPath;
+	UPROPERTY(EditAnywhere, BlueprintReadWrite, Category="Command Description")
+	FString Tooltip;
+	UPROPERTY(EditAnywhere, BlueprintReadWrite, Category="Command Description", AssetRegistrySearchable)
+	FString Category="Default";
+	UPROPERTY(EditAnywhere,BlueprintReadWrite, Category="Command Description")
+	bool bShowParameters=false;
+	
+	UPROPERTY(EditAnywhere, Category="Command Description")
+	FInputChord	KeyboardShortcut=FInputChord();
+
+	UPROPERTY(EditAnywhere,BlueprintReadWrite,meta=(AllowAbstract = "false", MustImplement = "/Script/UserToolBoxCore.UTBUICommand" ), Category="Command Description")
+	TObjectPtr<UClass>	UIOverride=nullptr;
+	UPROPERTY(Transient)
+	TObjectPtr<UObject>		UI=nullptr;
+	UPROPERTY()
+	bool bIsTransaction=false;
+	virtual void Execute();
+	UFUNCTION()
+	void ExecuteCommand();
+	virtual bool DisplayParameters();
+	virtual UUTBBaseCommand* CopyCommand(UObject* Owner) const ;
+
+	UFUNCTION(BlueprintCallable,Category="Transaction")
+	void AddObjectToTransaction(UObject* Object);
+	UFUNCTION(BlueprintCallable,Category="Transaction")
+	void AddObjectsToTransaction(TArray<UObject*> Objects);
+};
+
diff --git a/Engine/Plugins/Experimental/UserToolBoxCore/Source/UserToolBoxCore/Public/UTBBaseTab.h b/Engine/Plugins/Experimental/UserToolBoxCore/Source/UserToolBoxCore/Public/UTBBaseTab.h
new file mode 100644
--- /dev/null
+++ b/Engine/Plugins/Experimental/UserToolBoxCore/Source/UserToolBoxCore/Public/UTBBaseTab.h
@@ -0,0 +1,198 @@
+// Copyright Epic Games, Inc. All Rights Reserved.
+
+#pragma once
+
+#include "CoreMinimal.h"
+#include "UTBBaseUICommandInterface.h"
+#include "Engine/DataAsset.h"
+#include "Widgets/SCompoundWidget.h"
+#include "UTBBaseTab.generated.h"
+
+/**
+ * 
+ */
+class UUTBDefaultUITemplate;
+class UUTBBaseCommand;
+
+
+/**
+ * UTBTabSection is the class that represent a section in a UUserToolBoxBaseTab
+ * it contains the name of the section and the list of command attach to that section
+ */
+struct FUITemplateParameters
+{
+	bool	bForceHideSettingsButton=false;
+	bool	bPrefixSectionWithTabName=false;
+	bool	bPostfixSectionWithTabName=false;
+};
+
+UCLASS()
+class USERTOOLBOXCORE_API UUTBTabSection : public UObject
+{
+	GENERATED_BODY()
+public:
+	/** The name of the section*/
+	UPROPERTY(EditAnywhere, Category="Command Tab")
+	FString								SectionName;
+	/** The list of command */
+	UPROPERTY()
+	TArray<TObjectPtr<UUTBBaseCommand>>			Commands;
+};
+/**
+ * UUserToolBoxBaseTab is the class that represent a configurable tab
+ * This class contains the name of the tab and a section's list.
+ * In that class, you can also specified an UTBUICommand to override "by default" the UI of each command.
+ * Note: The override define in the command override the UI override define in the tab.
+ * This is the class serialized as an asset.
+ */
+UCLASS()
+class USERTOOLBOXCORE_API UUserToolBoxBaseTab : public UDataAsset
+{
+	GENERATED_BODY()
+public:
+	static const FName PlaceHolderSectionName;
+	/** The name of the toolbox */
+	UPROPERTY(EditAnywhere, Category="Command Tab", AssetRegistrySearchable)
+	FString Name = "DefaultUserToolBox";
+	/** The UI template to use to render this tab*/
+	UPROPERTY(EditAnywhere, Category="Command Tab")
+	TSubclassOf<UUTBDefaultUITemplate>	TabUI;
+	/** the default override for each command ( does not override ui command specified in the command) */
+	UPROPERTY(EditAnywhere, meta=(MustImplement = "/Script/UserToolBoxCore.UTBUICommand", AllowAbstract="false"), Category="Command Tab")
+	TObjectPtr<UClass>	DefaultCommandUIOverride;
+
+
+	/** Should the setting's icon be visible on the top right of the tab ( for easy access)*/
+	UPROPERTY(EditAnywhere, Category="Command Tab")
+	bool	IsSettingShouldBeVisible=true;
+	UPROPERTY(EditAnywhere, Category="Command Tab",AssetRegistrySearchable)
+	bool	bIsVisibleInViewportOverlay=false;
+	UPROPERTY(EditAnywhere, Category="Command Tab",AssetRegistrySearchable)
+	bool	bIsVisibleInDrawer=true;
+	UPROPERTY(EditAnywhere, Category="Command Tab", AssetRegistrySearchable)
+	bool bIsVisibleInWindowsMenu =true;
+	/** Function to generate the raw widget of the tab. you may consider using the SUserToolBoxTabWidget instead if you want to register your different commands to the input manager etc..*/
+	TSharedPtr<SWidget>  GenerateUI();
+	TSharedPtr<SWidget>  GenerateUI(const TSubclassOf<UUTBDefaultUITemplate> TabUiClass,const FUITemplateParameters& Params);
+	virtual void PostLoad() override;
+	virtual void BeginDestroy() override;
+	/** Insert a default command from the class at a specific location in a specific section
+	 * @param InClass the class we want to instanciate
+	 * @param Section the section name
+	 * @param Index  the position in the section of the new command
+	 */
+	void InsertCommandFromClass(const TSubclassOf<UUTBBaseCommand> InClass,const  FString Section,const int Index=-1);
+	/** Insert a command at a specific location in a specific section
+	 * @param InCommand the command we want to add to the tab
+	 * @param Section the section name
+	 * @param Index the position in the section of the new command
+	 */
+	void InsertCommand(const UUTBBaseCommand* InCommand,const  FString SectionName,const int Index=-1);
+	/** remove a command at a specific location in a specific section
+	* @param Section the section name
+	* @param Index the position in the section of the command to remove
+	*/
+	void RemoveCommand(const FString SectionName,const int Index=0);
+	/** return the section with the specified name
+	* @param SectionName the section name
+	*/
+	UUTBTabSection* GetSection(const FString SectionName);
+	/** return the section with the specified name or create it if it does not exist
+	* @param SectionName the section name
+	*/
+	UUTBTabSection* GetOrCreateSection(const FString SectionName);
+	/** remove a section
+	* @param Section The section to remove
+	*/
+	void RemoveSection( UUTBTabSection* Section);
+	/** remove a section with a specific name
+	* @param SectionName The name of the section to remove
+	*/
+	void RemoveSection(const FString SectionName);
+	/** remove a command
+	* @param Command The command to remove
+	*/
+	void RemoveCommand(UUTBBaseCommand* Command);
+	/** remove a command contained in a specific section
+	* @param Command The name of the section to remove
+	* @param SectionName The name of the section to remove
+	*/
+	void RemoveCommandFromSection(UUTBBaseCommand* Command,const FString SectionName);
+	/** replace a batch of command
+	* @param ReplacementMap Map that contains the command to replace
+	*/
+	void ReplaceCommands(const TMap<UObject*,UObject*>& ReplacementMap);
+	/** get all the section of that tab
+	* @return an array of section 
+	*/
+	TArray<UUTBTabSection*>	GetSections();
+
+	void MoveSectionAfterExistingSection(const FString SectionToMoveName,const FString SectionToTargetName);
+	/** return the placeholder section
+	* @return the placeholder section 
+	*/
+	UUTBTabSection* GetPlaceHolderSection();
+	/** return the presence of the command in the tab
+	* @return is the command inside that tab 
+	*/
+	bool ContainsCommand(const UUTBBaseCommand* Command) const;
+
+	/** registering the command list
+	*/
+	void RegisterCommand();
+	/** unregistering the command list
+	*/
+	void UnregisterCommand();
+
+	void RemoveInvalidCommand();
+private:
+	// the array of sections
+	UPROPERTY()
+	TArray<TObjectPtr<UUTBTabSection>>		Sections;
+	// the placeholder section
+	UPROPERTY()
+	TObjectPtr<UUTBTabSection>				PlaceholderSection;
+
+
+	FDelegateHandle OnObjectReplacedHandle;
+};
+
+/**
+ * SUserToolBoxTabWidget is the UI widget generated by a UUserToolBoxBaseTab.
+ * It will generated the wanted UI based on the different UI override set in the UUserToolBoxBaseTab or in a specific command
+ */
+class USERTOOLBOXCORE_API SUserToolBoxTabWidget : public SCompoundWidget
+{
+	SLATE_BEGIN_ARGS( SUserToolBoxTabWidget )
+		: _Tab(nullptr),
+		 _UIOverride(nullptr)
+
+	{
+		
+	}
+	/** Sets the text content for this editable text widget */
+	SLATE_ATTRIBUTE( UUserToolBoxBaseTab*, Tab )
+	SLATE_ATTRIBUTE( TSubclassOf<UUTBDefaultUITemplate>, UIOverride )
+	SLATE_ATTRIBUTE( FUITemplateParameters, UIParameters )
+SLATE_END_ARGS()
+
+	
+	~SUserToolBoxTabWidget();
+	void Construct( const FArguments& InArgs );
+	/*
+	 * Function to request a tab UI update
+	 * @param Tab Tab to update
+	 */
+	void UpdateTab(UUserToolBoxBaseTab* InTab);
+	
+	void OnObjectModified( UObject* Object, struct FPropertyChangedEvent&);
+	
+private:
+	TObjectPtr<UUserToolBoxBaseTab> Tab;
+	TSubclassOf<UUTBDefaultUITemplate> UIOverride;
+	FUITemplateParameters Parameters;
+	//Delegate handle
+	FDelegateHandle		OnTabChangedDelegate;
+	FDelegateHandle		OnObjectPropertyChangedDelegate;
+	
+};
diff --git a/Engine/Plugins/Experimental/UserToolBoxCore/Source/UserToolBoxCore/Public/UTBBaseUICommand.h b/Engine/Plugins/Experimental/UserToolBoxCore/Source/UserToolBoxCore/Public/UTBBaseUICommand.h
new file mode 100644
--- /dev/null
+++ b/Engine/Plugins/Experimental/UserToolBoxCore/Source/UserToolBoxCore/Public/UTBBaseUICommand.h
@@ -0,0 +1,58 @@
+// Copyright Epic Games, Inc. All Rights Reserved.
+
+#pragma once
+
+#include "CoreMinimal.h"
+
+
+#include "Blueprint/UserWidget.h"
+#include "UTBBaseUICommandInterface.h"
+#include "UTBBaseUICommand.generated.h"
+
+class UUTBBaseCommand;
+//Native base class for Command UI Customization
+UCLASS()
+class USERTOOLBOXCORE_API UBaseCommandNativeUI : public  UObject, public IUTBUICommand
+{
+	GENERATED_BODY()
+	public:
+	virtual bool IsSupportingCommandClass(TSubclassOf<UUTBBaseCommand> CommandClass) override;
+	virtual void SetCommand(UUTBBaseCommand* Command) override ;
+	virtual void ExecuteCurrentCommand() override ;
+	virtual TSharedRef<SWidget> GetUI() override ;
+
+
+	protected:
+	UPROPERTY()
+	TObjectPtr<UUTBBaseCommand>	MyCommand;
+
+public:
+	void ReplaceCommand(const TMap<UObject*,UObject*>& Map);
+	virtual void BeginDestroy() override;
+};
+
+UCLASS()
+class USERTOOLBOXCORE_API UUTBCommandUMGUI : public UUserWidget , public IUTBUICommand
+{
+	GENERATED_BODY()
+
+	public:
+	UFUNCTION(BlueprintImplementableEvent, Category=" Command UI")
+	bool DoesSupportClass(TSubclassOf<UUTBBaseCommand> ObjectClass);
+
+	UFUNCTION(BlueprintCallable, Category=" Command UI")
+	void ExecuteCommand();
+
+	UPROPERTY(EditAnywhere,BlueprintReadOnly, Category=" Command UI")
+	TObjectPtr<UUTBBaseCommand> Command;
+
+
+
+	virtual bool IsSupportingCommandClass(TSubclassOf<UUTBBaseCommand> CommandClass) override;
+	virtual void SetCommand(UUTBBaseCommand* Command) override;
+	virtual void ExecuteCurrentCommand() override;
+	virtual TSharedRef<SWidget> GetUI() override;
+
+	void ReplaceCommand(const TMap<UObject*,UObject*>& Map);
+	virtual void BeginDestroy() override;
+};
diff --git a/Engine/Plugins/Experimental/UserToolBoxCore/Source/UserToolBoxCore/Public/UTBBaseUICommandInterface.h b/Engine/Plugins/Experimental/UserToolBoxCore/Source/UserToolBoxCore/Public/UTBBaseUICommandInterface.h
new file mode 100644
--- /dev/null
+++ b/Engine/Plugins/Experimental/UserToolBoxCore/Source/UserToolBoxCore/Public/UTBBaseUICommandInterface.h
@@ -0,0 +1,29 @@
+// Copyright Epic Games, Inc. All Rights Reserved.
+
+#pragma once
+
+#include "CoreMinimal.h"
+#include "UObject/Interface.h"
+#include "Blueprint/UserWidget.h"
+#include "UTBBaseUICommandInterface.generated.h"
+
+class UUTBBaseCommand;
+
+
+//Interface shared by the native UI and by the UMG one
+UINTERFACE(MinimalAPI)
+class UUTBUICommand : public UInterface
+{
+	GENERATED_BODY()
+};
+
+class USERTOOLBOXCORE_API IUTBUICommand
+{
+	GENERATED_BODY()
+	public:
+	virtual bool IsSupportingCommandClass(TSubclassOf<UUTBBaseCommand> CommandClass)=0;
+	virtual void SetCommand(UUTBBaseCommand* Command)=0;
+	virtual void ExecuteCurrentCommand()=0;
+	virtual TSharedRef<SWidget> GetUI()=0;
+
+};
diff --git a/Engine/Plugins/Experimental/UserToolBoxCore/Source/UserToolBoxCore/Public/UTBBaseUITab.h b/Engine/Plugins/Experimental/UserToolBoxCore/Source/UserToolBoxCore/Public/UTBBaseUITab.h
new file mode 100644
--- /dev/null
+++ b/Engine/Plugins/Experimental/UserToolBoxCore/Source/UserToolBoxCore/Public/UTBBaseUITab.h
@@ -0,0 +1,65 @@
+// Copyright Epic Games, Inc. All Rights Reserved.
+
+#pragma once
+
+#include "CoreMinimal.h"
+#include "UTBBaseUICommandInterface.h"
+#include "Framework/MultiBox/MultiBoxBuilder.h"
+#include "UObject/NoExportTypes.h"
+#include "UTBBaseUITab.generated.h"
+
+class UUTBUMGCommandUI;
+class UUserToolBoxBaseTab;
+class UUTBBaseCommand;
+struct FUITemplateParameters;
+
+
+
+/**
+ * 
+ */
+UCLASS()
+class USERTOOLBOXCORE_API UUTBDefaultUITemplate : public UObject
+{
+	GENERATED_BODY()
+	public:
+	virtual TSharedPtr<SWidget> BuildTabUI(UUserToolBoxBaseTab* Tab, const FUITemplateParameters& Params);
+	
+	
+
+};
+UCLASS()
+class USERTOOLBOXCORE_API UUTBToolBarTabUI : public UUTBDefaultUITemplate
+{
+	GENERATED_BODY()
+public:
+	virtual TSharedPtr<SWidget> BuildTabUI(UUserToolBoxBaseTab* Tab, const FUITemplateParameters& Params) override;
+	
+protected:
+	virtual TSharedPtr<FToolBarBuilder> GetBuilder(TSharedPtr<const FUICommandList> InCommandList, FMultiBoxCustomization InCustomization, TSharedPtr<FExtender> InExtender = nullptr, const bool InForceSmallIcons = false);
+
+};
+UCLASS()
+class USERTOOLBOXCORE_API UUTBPaletteTabUI : public UUTBToolBarTabUI
+{
+	GENERATED_BODY()
+protected:
+	virtual TSharedPtr<FToolBarBuilder> GetBuilder(TSharedPtr<const FUICommandList> InCommandList, FMultiBoxCustomization InCustomization, TSharedPtr<FExtender> InExtender = nullptr, const bool InForceSmallIcons = false)override;
+
+};
+UCLASS()
+class USERTOOLBOXCORE_API UUTBVerticalToolBarTabUI : public UUTBToolBarTabUI
+{
+	GENERATED_BODY()
+protected:
+	virtual TSharedPtr<FToolBarBuilder> GetBuilder(TSharedPtr<const FUICommandList> InCommandList, FMultiBoxCustomization InCustomization, TSharedPtr<FExtender> InExtender = nullptr, const bool InForceSmallIcons = false)override;
+
+};
+UCLASS()
+class USERTOOLBOXCORE_API UUTBSlimHorizontalToolBarTabUI : public UUTBToolBarTabUI
+{
+	GENERATED_BODY()
+protected:
+	virtual TSharedPtr<FToolBarBuilder> GetBuilder(TSharedPtr<const FUICommandList> InCommandList, FMultiBoxCustomization InCustomization, TSharedPtr<FExtender> InExtender = nullptr, const bool InForceSmallIcons = false)override;
+
+};
\ No newline at end of file
diff --git a/Engine/Plugins/Experimental/UserToolBoxCore/Source/UserToolBoxCore/Public/UTBEditorBlueprint.h b/Engine/Plugins/Experimental/UserToolBoxCore/Source/UserToolBoxCore/Public/UTBEditorBlueprint.h
new file mode 100644
--- /dev/null
+++ b/Engine/Plugins/Experimental/UserToolBoxCore/Source/UserToolBoxCore/Public/UTBEditorBlueprint.h
@@ -0,0 +1,17 @@
+// Copyright Epic Games, Inc. All Rights Reserved.
+
+#pragma once
+
+#include "CoreMinimal.h"
+#include "EditorUtilityBlueprint.h"
+#include "UTBEditorBlueprint.generated.h"
+
+/**
+ * 
+ */
+UCLASS()
+class USERTOOLBOXCORE_API UUTBEditorBlueprint : public UEditorUtilityBlueprint
+{
+	GENERATED_BODY()
+	
+};
diff --git a/Engine/Plugins/Experimental/UserToolBoxCore/Source/UserToolBoxCore/Public/UTBFactory.h b/Engine/Plugins/Experimental/UserToolBoxCore/Source/UserToolBoxCore/Public/UTBFactory.h
new file mode 100644
--- /dev/null
+++ b/Engine/Plugins/Experimental/UserToolBoxCore/Source/UserToolBoxCore/Public/UTBFactory.h
@@ -0,0 +1,46 @@
+// Copyright Epic Games, Inc. All Rights Reserved.
+
+#pragma once
+
+#include "CoreMinimal.h"
+#include "Factories/Factory.h"
+#include "UTBFactory.generated.h"
+
+/**
+ * 
+ */
+UCLASS()
+class USERTOOLBOXCORE_API UUTBEditorUtilityBlueprintFactory : public UFactory
+{
+	GENERATED_UCLASS_BODY()
+	virtual UObject* FactoryCreateNew(UClass* InClass, UObject* InParent, FName InName, EObjectFlags Flags, UObject* Context, FFeedbackContext* Warn) override;
+	virtual bool CanCreateNew() const override;
+};
+
+UCLASS()
+class USERTOOLBOXCORE_API UUTBCommandFactory : public UFactory
+{
+	GENERATED_UCLASS_BODY()
+	virtual UObject* FactoryCreateNew(UClass* InClass, UObject* InParent, FName InName, EObjectFlags Flags, UObject* Context, FFeedbackContext* Warn) override;
+	virtual bool CanCreateNew() const override;
+	virtual bool ConfigureProperties() override;
+	
+	private:
+	UClass* ParentClass;
+};
+UCLASS()
+class USERTOOLBOXCORE_API UUTBCommandTabFactory : public UFactory
+{
+	GENERATED_UCLASS_BODY()
+	virtual UObject* FactoryCreateNew(UClass* InClass, UObject* InParent, FName InName, EObjectFlags Flags, UObject* Context, FFeedbackContext* Warn) override;
+	virtual bool CanCreateNew() const override;
+
+};
+UCLASS()
+class USERTOOLBOXCORE_API UIconTrackerFactory : public UFactory
+{
+	GENERATED_UCLASS_BODY()
+	virtual UObject* FactoryCreateNew(UClass* InClass, UObject* InParent, FName InName, EObjectFlags Flags, UObject* Context, FFeedbackContext* Warn) override;
+	virtual bool CanCreateNew() const override;
+
+};
diff --git a/Engine/Plugins/Experimental/UserToolBoxCore/Source/UserToolBoxCore/Public/UserToolBoxBaseBlueprint.h b/Engine/Plugins/Experimental/UserToolBoxCore/Source/UserToolBoxCore/Public/UserToolBoxBaseBlueprint.h
new file mode 100644
--- /dev/null
+++ b/Engine/Plugins/Experimental/UserToolBoxCore/Source/UserToolBoxCore/Public/UserToolBoxBaseBlueprint.h
@@ -0,0 +1,28 @@
+// Copyright Epic Games, Inc. All Rights Reserved.
+
+#pragma once
+
+#include "CoreMinimal.h"
+#include "UTBBaseCommand.h"
+#include "UObject/NoExportTypes.h"
+#include "UserToolBoxBaseBlueprint.generated.h"
+
+/**
+ * 
+ */
+
+
+UCLASS(Abstract, hideCategories=(Object),Blueprintable, meta = (ShowWorldContextPin))
+class USERTOOLBOXCORE_API UUserToolBoxBaseBlueprint : public UUTBBaseCommand
+{
+	GENERATED_BODY()
+	public:
+
+
+	UFUNCTION(BlueprintImplementableEvent)
+	void Command() ;
+
+	virtual void Execute() override;
+};
+
+
diff --git a/Engine/Plugins/Experimental/UserToolBoxCore/Source/UserToolBoxCore/Public/UserToolBoxCore.h b/Engine/Plugins/Experimental/UserToolBoxCore/Source/UserToolBoxCore/Public/UserToolBoxCore.h
new file mode 100644
--- /dev/null
+++ b/Engine/Plugins/Experimental/UserToolBoxCore/Source/UserToolBoxCore/Public/UserToolBoxCore.h
@@ -0,0 +1,16 @@
+// Copyright Epic Games, Inc. All Rights Reserved.
+
+#pragma once
+
+#include "CoreMinimal.h"
+#include "Modules/ModuleInterface.h"
+DECLARE_LOG_CATEGORY_EXTERN(LogUserToolBoxCore, Log, All);
+
+class FUserToolBoxCoreModule : public IModuleInterface
+{
+public:
+
+	/** IModuleInterface implementation */
+	virtual void StartupModule() override;
+	virtual void ShutdownModule() override;
+};
diff --git a/Engine/Plugins/Experimental/UserToolBoxCore/Source/UserToolBoxCore/Public/UserToolBoxStyle.h b/Engine/Plugins/Experimental/UserToolBoxCore/Source/UserToolBoxCore/Public/UserToolBoxStyle.h
new file mode 100644
--- /dev/null
+++ b/Engine/Plugins/Experimental/UserToolBoxCore/Source/UserToolBoxCore/Public/UserToolBoxStyle.h
@@ -0,0 +1,37 @@
+// Copyright Epic Games, Inc. All Rights Reserved.
+
+#pragma once
+
+#include "CoreMinimal.h"
+ 	
+
+#include "Styling/ISlateStyle.h"
+#include "IconsTracker.h"
+/**  */
+class USERTOOLBOXCORE_API FUserToolBoxStyle
+{
+public:
+
+	static void Initialize();
+
+	static void Shutdown();
+
+	/** reloads textures used by slate renderer */
+	static void ReloadTextures();
+
+	/** @return The Slate style set for the Shooter game */
+	static const ISlateStyle& Get();
+
+	static FName GetStyleSetName();
+	static void AddExternalImageBrushes(const TArray<FIconInfo>& IconInfos);
+	static void ClearExternalImageBrushes();
+	static TArray<FString> GetAvailableExternalImageBrushes();
+private:
+
+	static TSharedRef< class FSlateStyleSet > Create();
+
+
+
+	static TSharedPtr< class FSlateStyleSet > StyleInstance;
+	static TArray<FString> ExternalBrushIds;
+};
\ No newline at end of file
diff --git a/Engine/Plugins/Experimental/UserToolBoxCore/Source/UserToolBoxCore/Public/UserToolBoxSubsystem.h b/Engine/Plugins/Experimental/UserToolBoxCore/Source/UserToolBoxCore/Public/UserToolBoxSubsystem.h
new file mode 100644
--- /dev/null
+++ b/Engine/Plugins/Experimental/UserToolBoxCore/Source/UserToolBoxCore/Public/UserToolBoxSubsystem.h
@@ -0,0 +1,73 @@
+// Copyright Epic Games, Inc. All Rights Reserved.
+
+#pragma once
+
+#include "CoreMinimal.h"
+#include "AssetTypeCategories.h"
+#include "EditorSubsystem.h"
+#include "UTBBaseTab.h"
+#include "UTBBaseUITab.h"
+#include "UserToolBoxSubsystem.generated.h"
+
+class UUTBCommandUMGUI;
+class UUTBBaseCommand;
+class UUTBBaseUITab;
+/**
+ * 
+ */
+DECLARE_MULTICAST_DELEGATE_OneParam(FOnTabChanged, UUserToolBoxBaseTab*)
+UCLASS(BlueprintType)
+class USERTOOLBOXCORE_API UUserToolboxSubsystem : public UEditorSubsystem
+{
+	GENERATED_BODY()
+	UFUNCTION(BlueprintCallable, Category="User Toolbox Subsystem")
+	void RegisterTabData();
+
+	
+
+	public:
+	TArray<FAssetData> GetAvailableTabList();
+	FOnTabChanged		OnTabChanged;
+
+	UFUNCTION(BlueprintCallable,Category="User Toolbox")
+	bool PickAnIcon(FString& OutValue);
+	
+	private:
+	UPROPERTY()
+	TMap<FString,TObjectPtr<UUserToolBoxBaseTab>>	RegisteredTabs;
+	
+	
+	public:
+	UFUNCTION(BlueprintCallable,Category="User Toolbox")
+	void RefreshIcons();
+
+	
+	
+	TSharedPtr<SWidget> GenerateTabUI(const FAssetData Data, TSubclassOf<UUTBDefaultUITemplate> Ui=nullptr,const FUITemplateParameters& Parameters=FUITemplateParameters());
+
+	
+	/** Implement this for initialization of instances of the system */
+	virtual void Initialize(FSubsystemCollectionBase& Collection) override;
+
+	/** Implement this for deinitialization of instances of the system */
+	virtual void Deinitialize() override;
+
+	
+	static EAssetTypeCategories::Type AssetCategory;
+	private:
+
+	TArray<TPair<FString,FString>> Icons;
+
+	TArray<TSharedPtr<FString>>		IconOptions;
+
+public:
+	void RegisterDrawer();
+	void UpdateLevelViewportWidget();
+private:
+
+	// Delegate handle
+	FDelegateHandle OnFileLoadHandle;
+
+	TSharedPtr<SWidget>	LevelViewportOverlayWdget;
+	TArray<FName>		RegisteredDrawer;
+};
diff --git a/Engine/Plugins/Experimental/UserToolBoxCore/Source/UserToolBoxCore/UserToolBoxCore.Build.cs b/Engine/Plugins/Experimental/UserToolBoxCore/Source/UserToolBoxCore/UserToolBoxCore.Build.cs
new file mode 100644
--- /dev/null
+++ b/Engine/Plugins/Experimental/UserToolBoxCore/Source/UserToolBoxCore/UserToolBoxCore.Build.cs
@@ -0,0 +1,87 @@
+// Copyright Epic Games, Inc. All Rights Reserved.
+
+using UnrealBuildTool;
+
+public class UserToolBoxCore : ModuleRules
+{
+	public UserToolBoxCore(ReadOnlyTargetRules Target) : base(Target)
+	{
+		
+		PCHUsage = ModuleRules.PCHUsageMode.UseExplicitOrSharedPCHs;
+		
+		PublicIncludePaths.AddRange(
+			new string[] {
+				// ... add public include paths required here ...
+			}
+			);
+				
+		
+		PrivateIncludePaths.AddRange(
+			new string[] {
+				// ... add other private include paths required here ...
+			}
+			);
+			
+		
+		PublicDependencyModuleNames.AddRange(
+			new string[]
+			{
+				"Core",
+				// ... add other public dependencies that you statically link with here ...
+			}
+			);
+			
+		
+		PrivateDependencyModuleNames.AddRange(
+			new string[]
+			{
+				"CoreUObject",
+				"Engine",
+				"Slate",
+				"SlateCore",
+				"UnrealEd",
+				"CoreUObject",
+				"Engine",
+				"Slate",
+				"SlateCore",
+				"EditorSubsystem",
+				"Blutility",
+				"ToolWidgets",
+				"AssetTools",
+				"PropertyEditor",
+				"UMG",
+				"UMGEditor",
+				"ToolMenus",
+				"Projects",
+				"LevelEditor",
+				"Kismet",
+				"InputCore",
+				"EditorStyle",
+				"EditorInteractiveToolsFramework",
+				"InteractiveToolsFramework",
+				"EditorWidgets", "SlateScriptingCommands"
+				,"StatusBar"
+				// ... add private dependencies that you statically link with here ...	
+			}
+			);
+		if (Target.Version.MajorVersion >= 5)
+		{
+			PrivateDependencyModuleNames.AddRange(
+				new string[]
+				{
+					"EditorFramework"
+				});
+		}
+		else
+		{
+		}
+
+
+		DynamicallyLoadedModuleNames.AddRange(
+			new string[]
+			{
+				// ... add any modules that your module loads dynamically here ...
+			}
+			);
+	}
+}
diff --git a/Engine/Plugins/Experimental/UserToolBoxCore/UserToolBoxCore.uplugin b/Engine/Plugins/Experimental/UserToolBoxCore/UserToolBoxCore.uplugin
new file mode 100644
--- /dev/null
+++ b/Engine/Plugins/Experimental/UserToolBoxCore/UserToolBoxCore.uplugin
@@ -0,0 +1,31 @@
+{
+	"FileVersion": 1,
+	"Version": 1,
+	"VersionName": "0.10",
+	"FriendlyName": "UserToolBoxCore",
+	"Description": "Core functionnality to create custom editor tab",
+	"Category": "Editor",
+	"CreatedBy": "Epic Games, Inc.",
+	"CreatedByURL": "https://epicgames.com",
+	"DocsURL": "",
+	"MarketplaceURL": "",
+	"SupportURL": "",
+	"CanContainContent": false,
+	"IsBetaVersion": false,
+	"IsExperimentalVersion": true,
+	"Installed": false,
+	"Modules": [
+		{
+			"Name": "UserToolBoxCore",
+			"Type": "Editor",
+			"LoadingPhase": "Default"
+		}
+	],
+	"Plugins":
+	[
+		{
+			"Name": "SlateScripting",
+			"Enabled": true
+		}
+	]
+}
\ No newline at end of file
