UUserToolboxSubsystem: Refresh icons when an UIconsTracker is added, removed, or updated.

diff --git a/Engine/Plugins/Experimental/UserToolBoxCore/Source/UserToolBoxCore/Private/UserToolBoxSubsystem.cpp b/Engine/Plugins/Experimental/UserToolBoxCore/Source/UserToolBoxCore/Private/UserToolBoxSubsystem.cpp
--- a/Engine/Plugins/Experimental/UserToolBoxCore/Source/UserToolBoxCore/Private/UserToolBoxSubsystem.cpp
+++ b/Engine/Plugins/Experimental/UserToolBoxCore/Source/UserToolBoxCore/Private/UserToolBoxSubsystem.cpp
@@ -362,15 +362,27 @@ void UUserToolboxSubsystem::Initialize(FSubsystemCollectionBase& Collection)
 	FAssetRegistryModule& AssetRegistryModule = FModuleManager::LoadModuleChecked<FAssetRegistryModule>("AssetRegistry");
 	IAssetRegistry& AssetRegistry = AssetRegistryModule.Get();
 	OnFileLoadHandle=AssetRegistry.OnFilesLoaded().AddLambda([this](){RegisterTabData();});
+
+	AssetRegistry.OnAssetAdded().AddUObject(this, &UUserToolboxSubsystem::OnAssetChanged);
+	AssetRegistry.OnAssetRemoved().AddUObject(this, &UUserToolboxSubsystem::OnAssetChanged);
+	AssetRegistry.OnAssetUpdatedOnDisk().AddUObject(this, &UUserToolboxSubsystem::OnAssetChanged);
+
 	FUserToolBoxStyle::Initialize();
 	FUTBEditorCommands::Register();
 }
+
 void UUserToolboxSubsystem::Deinitialize()
 {
+	FAssetRegistryModule& AssetRegistryModule = FModuleManager::GetModuleChecked<FAssetRegistryModule>("AssetRegistry");
+	IAssetRegistry& AssetRegistry = AssetRegistryModule.Get();
+
+	AssetRegistry.OnAssetAdded().RemoveAll(this);
+	AssetRegistry.OnAssetRemoved().RemoveAll(this);
+	AssetRegistry.OnAssetUpdatedOnDisk().RemoveAll(this);
+
 	if (OnFileLoadHandle.IsValid())
 	{
-		FAssetRegistryModule& AssetRegistryModule = FModuleManager::LoadModuleChecked<FAssetRegistryModule>("AssetRegistry");
-		IAssetRegistry& AssetRegistry = AssetRegistryModule.Get();
+		
 		AssetRegistry.OnFilesLoaded().Remove(OnFileLoadHandle);
 	}
 	FLevelEditorModule& LevelEditorModule = FModuleManager::GetModuleChecked<FLevelEditorModule>("LevelEditor");
@@ -596,3 +608,16 @@ void UUserToolboxSubsystem::UpdateLevelViewportWidget()
 	}
 }
 
+void UUserToolboxSubsystem::OnAssetChanged(const FAssetData& InAssetData)
+{
+	if (InAssetData.AssetClassPath != UIconsTracker::StaticClass()->GetClassPathName())
+	{
+		return;
+	}
+
+	UIconsTracker* Asset = Cast<UIconsTracker>(InAssetData.GetAsset());
+	if (IsValid(Asset))
+	{
+		RefreshIcons();
+	}
+}
diff --git a/Engine/Plugins/Experimental/UserToolBoxCore/Source/UserToolBoxCore/Public/UserToolBoxSubsystem.h b/Engine/Plugins/Experimental/UserToolBoxCore/Source/UserToolBoxCore/Public/UserToolBoxSubsystem.h
--- a/Engine/Plugins/Experimental/UserToolBoxCore/Source/UserToolBoxCore/Public/UserToolBoxSubsystem.h
+++ b/Engine/Plugins/Experimental/UserToolBoxCore/Source/UserToolBoxCore/Public/UserToolBoxSubsystem.h
@@ -65,6 +65,8 @@ public:
 	void UpdateLevelViewportWidget();
 private:
 
+	void OnAssetChanged(const FAssetData& InAssetData);
+
 	// Delegate handle
 	FDelegateHandle OnFileLoadHandle;
 
