Removed inline preventing compilation on MAC

diff --git a/Engine/Plugins/Experimental/UserToolBoxCore/Source/UserToolBoxCore/Private/Customization/BaseCommandCustomization.cpp b/Engine/Plugins/Experimental/UserToolBoxCore/Source/UserToolBoxCore/Private/Customization/BaseCommandCustomization.cpp
--- a/Engine/Plugins/Experimental/UserToolBoxCore/Source/UserToolBoxCore/Private/Customization/BaseCommandCustomization.cpp
+++ b/Engine/Plugins/Experimental/UserToolBoxCore/Source/UserToolBoxCore/Private/Customization/BaseCommandCustomization.cpp
@@ -75,7 +75,9 @@ void FBaseCommandCustomization::CustomizeDetails(IDetailLayoutBuilder& DetailBui
 	
 	
 }
-
+FBaseCommandCustomization::~FBaseCommandCustomization()
+{
+}
 void FBaseCommandCustomization::CustomizeDetails(const TSharedPtr<IDetailLayoutBuilder>& DetailBuilder)
 {
 	CustomizeDetails(*DetailBuilder);
diff --git a/Engine/Plugins/Experimental/UserToolBoxCore/Source/UserToolBoxCore/Private/Customization/BaseCommandCustomization.h b/Engine/Plugins/Experimental/UserToolBoxCore/Source/UserToolBoxCore/Private/Customization/BaseCommandCustomization.h
--- a/Engine/Plugins/Experimental/UserToolBoxCore/Source/UserToolBoxCore/Private/Customization/BaseCommandCustomization.h
+++ b/Engine/Plugins/Experimental/UserToolBoxCore/Source/UserToolBoxCore/Private/Customization/BaseCommandCustomization.h
@@ -25,8 +25,6 @@ public:
 	virtual void CustomizeDetails(const TSharedPtr<IDetailLayoutBuilder>& DetailBuilder) override;
 };
 
-inline FBaseCommandCustomization::~FBaseCommandCustomization()
-{
-}
+
 
 
diff --git a/Engine/Plugins/Experimental/UserToolBoxCore/Source/UserToolBoxCore/Private/Editor/UTBTabEditor.cpp b/Engine/Plugins/Experimental/UserToolBoxCore/Source/UserToolBoxCore/Private/Editor/UTBTabEditor.cpp
--- a/Engine/Plugins/Experimental/UserToolBoxCore/Source/UserToolBoxCore/Private/Editor/UTBTabEditor.cpp
+++ b/Engine/Plugins/Experimental/UserToolBoxCore/Source/UserToolBoxCore/Private/Editor/UTBTabEditor.cpp
@@ -633,7 +633,7 @@ void FUTBTabEditor::BindCommands()
 
 }
 
-inline void FUTBTabEditor::CreateOrRebuildUTBTabTabWidget()
+void FUTBTabEditor::CreateOrRebuildUTBTabTabWidget()
 {
 	UUserToolBoxBaseTab* Tab=Cast<UUserToolBoxBaseTab>(GetEditingObject());
 	if (!IsValid(Tab))
@@ -878,7 +878,7 @@ void FUTBTabEditor::CreateOrRebuildUTBTabSectionTabWidget()
 			.ClearSelectionOnClick(true)
 	];
 }
-inline TSharedPtr<SWidget> FUTBTabEditor::CreateUTBTabDetailsTabWidget()
+ TSharedPtr<SWidget> FUTBTabEditor::CreateUTBTabDetailsTabWidget()
 {
 	FPropertyEditorModule& PropertyEditorModule = FModuleManager::GetModuleChecked<FPropertyEditorModule>("PropertyEditor");
 	FDetailsViewArgs Args;
@@ -895,7 +895,7 @@ inline TSharedPtr<SWidget> FUTBTabEditor::CreateUTBTabDetailsTabWidget()
 	return 	DetailView;
 }
 
-inline TSharedPtr<SWidget> FUTBTabEditor::CreateUTBCommandTabWidget()
+ TSharedPtr<SWidget> FUTBTabEditor::CreateUTBCommandTabWidget()
 {
 	return SNew(SCommandPickMenuWidget);
 }
diff --git a/Engine/Plugins/Experimental/UserToolBoxCore/Source/UserToolBoxCore/Private/UserToolBoxBaseBlueprint.cpp b/Engine/Plugins/Experimental/UserToolBoxCore/Source/UserToolBoxCore/Private/UserToolBoxBaseBlueprint.cpp
--- a/Engine/Plugins/Experimental/UserToolBoxCore/Source/UserToolBoxCore/Private/UserToolBoxBaseBlueprint.cpp
+++ b/Engine/Plugins/Experimental/UserToolBoxCore/Source/UserToolBoxCore/Private/UserToolBoxBaseBlueprint.cpp
@@ -3,7 +3,7 @@
 
 #include "UserToolBoxBaseBlueprint.h"
 
-inline void UUserToolBoxBaseBlueprint::Execute()
+void UUserToolBoxBaseBlueprint::Execute()
 {
 	Command();
 }
\ No newline at end of file
