restore UserToolBoxTab
#11460 tag and #contrib 11460

diff --git a/Engine/Plugins/Experimental/UserToolBoxCore/Source/UserToolBoxCore/Private/UserToolBoxSubsystem.cpp b/Engine/Plugins/Experimental/UserToolBoxCore/Source/UserToolBoxCore/Private/UserToolBoxSubsystem.cpp
--- a/Engine/Plugins/Experimental/UserToolBoxCore/Source/UserToolBoxCore/Private/UserToolBoxSubsystem.cpp
+++ b/Engine/Plugins/Experimental/UserToolBoxCore/Source/UserToolBoxCore/Private/UserToolBoxSubsystem.cpp
@@ -153,7 +153,7 @@ TArray<FAssetData> UUserToolboxSubsystem::GetAvailableTabList()
 	Filter.bRecursiveClasses = true;
 	TArray<FAssetData> AssetDatas;
 	AssetRegistry.GetAssets(Filter,AssetDatas);
-	TArray<FAssetData> AvailableTabList;
+	TArray<FAssetData> TabList;
 	for (FAssetData AssetData:AssetDatas)
 	{
 		UUserToolBoxBaseTab* Tab=Cast<UUserToolBoxBaseTab>(AssetData.GetAsset());
@@ -162,9 +162,9 @@ TArray<FAssetData> UUserToolboxSubsystem::GetAvailableTabList()
 			UE_LOG(LogTemp,Warning,TEXT("Unable to load %s"),*AssetData.PackageName.ToString())
 			continue;
 		}
-		AvailableTabList.Add(AssetData);
+		TabList.Add(AssetData);
 	}
-	return AvailableTabList;
+	return TabList;
 }
 
 
@@ -369,6 +369,46 @@ void UUserToolboxSubsystem::Initialize(FSubsystemCollectionBase& Collection)
 
 	FUserToolBoxStyle::Initialize();
 	FUTBEditorCommands::Register();
+	
+	FLevelEditorModule& LevelEditorModule = FModuleManager::LoadModuleChecked<FLevelEditorModule>(TEXT("LevelEditor"));
+	LevelEditorModule.OnTabManagerChanged().AddUObject(this, &UUserToolboxSubsystem::ReinitializeUIs);
+}
+
+TSharedRef<SDockTab> UUserToolboxSubsystem::CreateTab(const FSpawnTabArgs& SpawnTabArgs)
+{
+	const FAssetData* TabAssetData =
+		AvailableTabList.FindByPredicate([SpawnTabArgs](const FAssetData& Tab)
+		{
+			return Tab.GetTagValueRef<FString>("Name") == SpawnTabArgs.GetTabId().ToString();
+		});
+	
+	return SNew(SDockTab)
+			.TabRole(NomadTab)
+			[
+				GEditor->GetEditorSubsystem<UUserToolboxSubsystem>()->GenerateTabUI(*TabAssetData).ToSharedRef()
+			];
+}
+
+void UUserToolboxSubsystem::ReinitializeUIs()
+{
+	AvailableTabList = GetAvailableTabList();
+	for (const FAssetData& TabAssetData : AvailableTabList)
+	{
+		FString TabName=TabAssetData.GetTagValueRef<FString>("Name");
+		if (RegisteredDrawer.Contains(FName(TabName)))
+		{
+			continue;	
+		}
+		
+		if (TabAssetData.GetTagValueRef<bool>("bIsVisibleInDrawer"))
+		{
+			FLevelEditorModule& LevelEditorModule = FModuleManager::GetModuleChecked<FLevelEditorModule>(TEXT("LevelEditor"));
+			if (const TSharedPtr<FTabManager> LevelEditorTabManager = LevelEditorModule.GetLevelEditorTabManager())
+			{
+				LevelEditorTabManager->RegisterTabSpawner(FName(TabName), FOnSpawnTab::CreateUObject(this, &UUserToolboxSubsystem::CreateTab));
+			}
+		}
+	}
 }
 
 void UUserToolboxSubsystem::Deinitialize()
@@ -417,7 +457,6 @@ void UUserToolboxSubsystem::Deinitialize()
 
 void UUserToolboxSubsystem::RegisterDrawer()
 {
-	
 	UStatusBarSubsystem* StatusBarSubsystem=GEditor->GetEditorSubsystem<UStatusBarSubsystem>();
 	for (FName TabToUnregister:RegisteredDrawer)
 	{
diff --git a/Engine/Plugins/Experimental/UserToolBoxCore/Source/UserToolBoxCore/Public/UserToolBoxSubsystem.h b/Engine/Plugins/Experimental/UserToolBoxCore/Source/UserToolBoxCore/Public/UserToolBoxSubsystem.h
--- a/Engine/Plugins/Experimental/UserToolBoxCore/Source/UserToolBoxCore/Public/UserToolBoxSubsystem.h
+++ b/Engine/Plugins/Experimental/UserToolBoxCore/Source/UserToolBoxCore/Public/UserToolBoxSubsystem.h
@@ -7,6 +7,7 @@
 #include "EditorSubsystem.h"
 #include "UTBBaseTab.h"
 #include "UTBBaseUITab.h"
+#include "Widgets/Docking/SDockTab.h"
 #include "UserToolBoxSubsystem.generated.h"
 
 class UUTBCommandUMGUI;
@@ -48,10 +49,11 @@ class USERTOOLBOXCORE_API UUserToolboxSubsystem : public UEditorSubsystem
 	
 	/** Implement this for initialization of instances of the system */
 	virtual void Initialize(FSubsystemCollectionBase& Collection) override;
-
 	/** Implement this for deinitialization of instances of the system */
 	virtual void Deinitialize() override;
-
+	
+	TSharedRef<SDockTab> CreateTab(const FSpawnTabArgs& SpawnTabArgs);
+	void ReinitializeUIs();
 	
 	static EAssetTypeCategories::Type AssetCategory;
 	private:
@@ -72,4 +74,5 @@ private:
 
 	TSharedPtr<SWidget>	LevelViewportOverlayWdget;
 	TArray<FName>		RegisteredDrawer;
+	TArray<FAssetData> AvailableTabList;
 };
