UserToolbox - Icon creation crashes engine
- Now creating FSlateImageBrushes and ReloadTextureResources only on changes
#Synced-CL 51657968
#Synced-CL 51657968

diff --git a/Engine/Plugins/Experimental/UserToolBoxCore/Source/UserToolBoxCore/Private/UserToolBoxStyle.cpp b/Engine/Plugins/Experimental/UserToolBoxCore/Source/UserToolBoxCore/Private/UserToolBoxStyle.cpp
--- a/Engine/Plugins/Experimental/UserToolBoxCore/Source/UserToolBoxCore/Private/UserToolBoxStyle.cpp
+++ b/Engine/Plugins/Experimental/UserToolBoxCore/Source/UserToolBoxCore/Private/UserToolBoxStyle.cpp
@@ -13,6 +13,8 @@
 
 TSharedPtr< FSlateStyleSet > FUserToolBoxStyle::StyleInstance = nullptr;
 TArray<FString> FUserToolBoxStyle::ExternalBrushIds;
+TMap<FString, TUniquePtr<FSlateImageBrush>> FUserToolBoxStyle::ExternalBrushes;
+
 void FUserToolBoxStyle::Initialize()
 {
 	if (!StyleInstance.IsValid())
@@ -62,17 +64,36 @@ void FUserToolBoxStyle::AddExternalImageBrushes(const TArray<FIconInfo>& IconInf
 	{
 		return;
 	}
+
+	bool bChanged = false;
+
 	for (const FIconInfo& Info:IconInfos)
 	{
-				StyleInstance->Set(FName(*Info.Id),new FSlateImageBrush(Info.Path,Info.IconSize));
-				ExternalBrushIds.Add(Info.Id);
+		const FString Key = FString::Printf(TEXT("%s_%s"), *Info.Path, *Info.IconSize.ToString());
+
+		TUniquePtr<FSlateImageBrush>* FoundBrush = ExternalBrushes.Find(Key);
+
+		if (!FoundBrush)
+		{
+			TUniquePtr<FSlateImageBrush> Brush = MakeUnique<FSlateImageBrush>(Info.Path, Info.IconSize);
+			ExternalBrushes.Add(Key, MoveTemp(Brush));
+			bChanged = true;
+		}
+
+		FoundBrush = ExternalBrushes.Find(Key);
+
+		if (FoundBrush)
+		{
+			StyleInstance->Set(*Info.Id, FoundBrush->Get());
+
+			ExternalBrushIds.Add(Info.Id);
+		}
 	}
-	ReloadTextures();
-}
 
-void FUserToolBoxStyle::ClearExternalImageBrushes()
-{
-	ExternalBrushIds.Empty();
+	if (bChanged)
+	{
+		ReloadTextures();
+	}
 }
 
 TArray<FString> FUserToolBoxStyle::GetAvailableExternalImageBrushes()
@@ -84,7 +105,10 @@ void FUserToolBoxStyle::ReloadTextures()
 {
 	if (FSlateApplication::IsInitialized())
 	{
-		FSlateApplication::Get().GetRenderer()->ReloadTextureResources();
+		if (auto Renderer = FSlateApplication::Get().GetRenderer())
+		{
+			Renderer->ReloadTextureResources();
+		}
 	}
 }
 
diff --git a/Engine/Plugins/Experimental/UserToolBoxCore/Source/UserToolBoxCore/Private/UserToolBoxSubsystem.cpp b/Engine/Plugins/Experimental/UserToolBoxCore/Source/UserToolBoxCore/Private/UserToolBoxSubsystem.cpp
--- a/Engine/Plugins/Experimental/UserToolBoxCore/Source/UserToolBoxCore/Private/UserToolBoxSubsystem.cpp
+++ b/Engine/Plugins/Experimental/UserToolBoxCore/Source/UserToolBoxCore/Private/UserToolBoxSubsystem.cpp
@@ -325,7 +325,7 @@ void UUserToolboxSubsystem::RefreshIcons()
 			}
 		}
 	}
-	FUserToolBoxStyle::ClearExternalImageBrushes();
+
 	FUserToolBoxStyle::AddExternalImageBrushes(IconDatas);
 }
 TSharedPtr<SWidget> UUserToolboxSubsystem::GenerateTabUI(const FAssetData Data, TSubclassOf<UUTBDefaultUITemplate> Ui,const FUITemplateParameters& Parameters)
diff --git a/Engine/Plugins/Experimental/UserToolBoxCore/Source/UserToolBoxCore/Public/UserToolBoxStyle.h b/Engine/Plugins/Experimental/UserToolBoxCore/Source/UserToolBoxCore/Public/UserToolBoxStyle.h
--- a/Engine/Plugins/Experimental/UserToolBoxCore/Source/UserToolBoxCore/Public/UserToolBoxStyle.h
+++ b/Engine/Plugins/Experimental/UserToolBoxCore/Source/UserToolBoxCore/Public/UserToolBoxStyle.h
@@ -7,6 +7,8 @@
 
 #include "Styling/ISlateStyle.h"
 #include "IconsTracker.h"
+struct FSlateImageBrush;
+
 /**  */
 class USERTOOLBOXCORE_API FUserToolBoxStyle
 {
@@ -24,7 +26,6 @@ public:
 
 	static FName GetStyleSetName();
 	static void AddExternalImageBrushes(const TArray<FIconInfo>& IconInfos);
-	static void ClearExternalImageBrushes();
 	static TArray<FString> GetAvailableExternalImageBrushes();
 private:
 
@@ -34,4 +35,5 @@ private:
 
 	static TSharedPtr< class FSlateStyleSet > StyleInstance;
 	static TArray<FString> ExternalBrushIds;
+	static TMap<FString, TUniquePtr<FSlateImageBrush>> ExternalBrushes;
 };
\ No newline at end of file
